Motion Design: added data link plugin

diff --git a/Engine/Plugins/Experimental/DataLink/DataLink.uplugin b/Engine/Plugins/Experimental/DataLink/DataLink.uplugin
new file mode 100644
--- /dev/null
+++ b/Engine/Plugins/Experimental/DataLink/DataLink.uplugin
@@ -0,0 +1,53 @@
+{
+	"FileVersion": 3,
+	"Version": 1,
+	"VersionName": "1.0",
+	"FriendlyName": "DataLink",
+	"Description": "",
+	"Category": "Other",
+	"CreatedBy": "",
+	"CreatedByURL": "",
+	"DocsURL": "",
+	"MarketplaceURL": "",
+	"CanContainContent": false,
+	"IsBetaVersion": false,
+	"IsExperimentalVersion": true,
+	"Installed": false,
+	"Modules": [
+		{
+			"Name": "DataLink",
+			"Type": "Runtime",
+			"LoadingPhase": "Default"
+		},
+		{
+			"Name": "DataLinkEdGraph",
+			"Type": "UncookedOnly",
+			"LoadingPhase": "Default"
+		},
+		{
+			"Name": "DataLinkEditor",
+			"Type": "Editor",
+			"LoadingPhase": "Default"
+		},
+		{
+			"Name": "DataLinkHttp",
+			"Type": "Runtime",
+			"LoadingPhase": "Default"
+		},
+		{
+			"Name": "DataLinkJson",
+			"Type": "Runtime",
+			"LoadingPhase": "Default"
+		},
+		{
+			"Name": "DataLinkJsonEditor",
+			"Type": "Editor",
+			"LoadingPhase": "Default"
+		},
+		{
+			"Name": "DataLinkDataTable",
+			"Type": "Runtime",
+			"LoadingPhase": "Default"
+		}
+	]
+}
\ No newline at end of file
diff --git a/Engine/Plugins/Experimental/DataLink/Source/DataLink/DataLink.Build.cs b/Engine/Plugins/Experimental/DataLink/Source/DataLink/DataLink.Build.cs
new file mode 100644
--- /dev/null
+++ b/Engine/Plugins/Experimental/DataLink/Source/DataLink/DataLink.Build.cs
@@ -0,0 +1,28 @@
+// Copyright Epic Games, Inc. All Rights Reserved.
+
+using UnrealBuildTool;
+
+public class DataLink : ModuleRules
+{
+	public DataLink(ReadOnlyTargetRules Target) : base(Target)
+	{
+		PCHUsage = ModuleRules.PCHUsageMode.UseExplicitOrSharedPCHs;
+
+		PublicDependencyModuleNames.AddRange(
+			new string[]
+			{
+				"Core",
+				"CoreUObject",
+			}
+		);
+
+		PrivateDependencyModuleNames.AddRange(
+			new string[]
+			{
+				"Engine",
+			}
+		);
+
+		PublicDefinitions.Add("WITH_DATALINK_CONTEXT=!NO_LOGGING");
+	}
+}
diff --git a/Engine/Plugins/Experimental/DataLink/Source/DataLink/Private/DataLinkExecutor.cpp b/Engine/Plugins/Experimental/DataLink/Source/DataLink/Private/DataLinkExecutor.cpp
new file mode 100644
--- /dev/null
+++ b/Engine/Plugins/Experimental/DataLink/Source/DataLink/Private/DataLinkExecutor.cpp
@@ -0,0 +1,404 @@
+// Copyright Epic Games, Inc. All Rights Reserved.
+
+#include "DataLinkExecutor.h"
+#include "DataLinkEnums.h"
+#include "DataLinkExecutorArguments.h"
+#include "DataLinkGraph.h"
+#include "DataLinkLog.h"
+#include "DataLinkNode.h"
+#include "DataLinkPinReference.h"
+#include "DataLinkSink.h"
+#include "Templates/SubclassOf.h"
+#include "UObject/Class.h"
+
+TSharedPtr<FDataLinkExecutor> FDataLinkExecutor::Create(FDataLinkExecutorArguments&& InArgs)
+{
+	TSharedRef<FDataLinkExecutor> Executor = MakeShared<FDataLinkExecutor>(FPrivateToken());
+
+#if WITH_DATALINK_CONTEXT
+	Executor->ContextName = MoveTemp(InArgs.ContextName);
+#endif
+	Executor->ContextObject = MoveTemp(InArgs.ContextObject);
+	Executor->Instance = MoveTemp(InArgs.Instance);
+	Executor->Sink = MoveTemp(InArgs.Sink);
+	Executor->OnFinishDelegate = MoveTemp(InArgs.OnFinishDelegate);
+
+	return Executor;
+}
+
+FDataLinkExecutor::FDataLinkExecutor(FPrivateToken)
+	: ExecutorId(FGuid::NewGuid())
+{
+}
+
+const FGuid& FDataLinkExecutor::GetExecutorId() const
+{
+	return ExecutorId;
+}
+
+FStringView FDataLinkExecutor::GetContextName() const
+{
+#if WITH_DATALINK_CONTEXT
+	return ContextName;
+#else
+	return FStringView();
+#endif
+}
+
+UObject* FDataLinkExecutor::GetContextObject() const
+{
+	return ContextObject;
+}
+
+FString FDataLinkExecutor::GetReferencerName() const
+{
+	return TEXT("FDataLinkExecutor");
+}
+
+void FDataLinkExecutor::AddReferencedObjects(FReferenceCollector& InCollector)
+{
+	InCollector.AddReferencedObject(ContextObject);
+	InCollector.AddPropertyReferencesWithStructARO(FDataLinkInstance::StaticStruct(), &Instance);
+
+	if (Sink.IsValid())
+	{
+		Sink->AddStructReferencedObjects(InCollector);
+	}
+
+	for (TPair<TObjectPtr<const UDataLinkNode>, FDataLinkNodeInstance>& Pair : NodeInstanceMap)
+	{
+		Pair.Value.AddReferencedObjects(InCollector);
+	}
+}
+
+void FDataLinkExecutor::Run()
+{
+	bRunning = true;
+
+	if (!ValidateRun())
+	{
+		Finish(EDataLinkExecutionResult::Failed, FConstStructView());
+		return;
+	}
+
+	// Ensure there's a valid sink instance
+	if (!Sink.IsValid())
+	{
+		Sink = MakeShared<FDataLinkSink>();
+	}
+
+	// Hold onto an extra shared reference in case the execution finishes immediately
+	TSharedRef<FDataLinkExecutor> This = SharedThis(this);
+
+	if (!ExecuteEntryNodes())
+	{
+		Finish(EDataLinkExecutionResult::Failed, FConstStructView());
+	}
+}
+
+const FDataLinkNodeInstance& FDataLinkExecutor::GetNodeInstance(const UDataLinkNode* InNode) const
+{
+	return NodeInstanceMap.FindChecked(InNode);
+}
+
+FDataLinkNodeInstance& FDataLinkExecutor::GetNodeInstanceMutable(const UDataLinkNode* InNode)
+{
+	return NodeInstanceMap.FindChecked(InNode);
+}
+
+void FDataLinkExecutor::SucceedNode(const UDataLinkNode* InNode, FConstStructView InOutputDataView)
+{
+	if (!ProcessFinishedNode(InNode))
+	{
+		Finish(EDataLinkExecutionResult::Failed, FConstStructView());
+		return;
+	}
+
+	TConstArrayView<FDataLinkPin> OutputPins = InNode->GetOutputPins();
+
+	TArray<FDataLinkPinReference> InputPinsToExecute;
+	InputPinsToExecute.Reserve(OutputPins.Num());
+
+	for (const FDataLinkPin& OutputPin : OutputPins)
+	{
+		if (const FDataLinkPin* InputPin = OutputPin.GetLinkedInputPin())
+		{
+			// If the Linked Input Pin was gotten successfully, then it means that the Output Pin's Linked Node was valid
+			check(OutputPin.LinkedNode);
+			InputPinsToExecute.Emplace(OutputPin.LinkedNode, InputPin);
+		}
+	}
+
+	// No next node, finish execution
+	if (InputPinsToExecute.IsEmpty())
+	{
+		// If the next node is null, it should be the output node
+		if (!ensureAlways(Instance.DataLinkGraph && Instance.DataLinkGraph->GetOutputNode() == InNode))
+		{
+			UE_LOG(LogDataLink, Error, TEXT("[%s] Node ('%s') has no node to go to and it's not an output node!"), GetContextName().GetData(), *InNode->GetName());
+			Finish(EDataLinkExecutionResult::Failed, FConstStructView());
+			return;
+		}
+
+		Finish(EDataLinkExecutionResult::Succeeded, InOutputDataView);
+		return;
+	}
+
+	TArray<FConstStructView> InputDataViews;
+	InputDataViews.Add(InOutputDataView);
+
+	if (!ExecuteInputPins(InputPinsToExecute, MoveTemp(InputDataViews)))
+	{
+		Finish(EDataLinkExecutionResult::Failed, FConstStructView());
+	}
+}
+
+void FDataLinkExecutor::FailNode(const UDataLinkNode* InNode)
+{
+	ProcessFinishedNode(InNode);
+
+	UE_LOG(LogDataLink, Log, TEXT("[%s] Node ('%s') has failed to execute. Data Link execution finished."), GetContextName().GetData(), *InNode->GetName());
+	Finish(EDataLinkExecutionResult::Failed, FConstStructView());
+}
+
+bool FDataLinkExecutor::ValidateRun()
+{
+	if (!Instance.DataLinkGraph)
+	{
+		UE_LOG(LogDataLink, Error, TEXT("[%s] Invalid Data Link specified!"), GetContextName().GetData());
+		return false;
+	}
+
+	const int32 InputPinCount = Instance.DataLinkGraph->GetInputPinCount();
+
+	if (Instance.InputData.Num() != InputPinCount)
+	{
+		UE_LOG(LogDataLink, Error, TEXT("[%s] DataLink Graph '%s' requires %d inputs but %d were provided")
+			, GetContextName().GetData()
+			, *Instance.DataLinkGraph->GetName()
+			, InputPinCount
+			, Instance.InputData.Num());
+		return false;
+	}
+
+	int32 PinIndex = 0;
+
+	// Ensure the input pins are compatible with the input data
+	const bool bResult = Instance.DataLinkGraph->ForEachInputPin(
+		[&PinIndex, this](FDataLinkPinReference InPinReference)->bool
+		{
+			const int32 CurrentPinIndex = PinIndex++;
+
+			if (InPinReference.Pin->Struct && InPinReference.Pin->Struct != Instance.InputData[CurrentPinIndex].GetScriptStruct())
+			{
+				UE_LOG(LogDataLink, Error, TEXT("[%s] Input pin (%s': input struct '%s', owner '%s') is not compatible with input data view '%s'")
+					, GetContextName().GetData()
+					, *InPinReference.Pin->Name.ToString()
+					, *InPinReference.Pin->Struct->GetName()
+					, *InPinReference.OwningNode->GetName()
+					, *GetNameSafe(Instance.InputData[CurrentPinIndex].GetScriptStruct()));
+				return false;
+			}
+
+			return true;
+		});
+
+	return bResult;
+}
+
+bool FDataLinkExecutor::ValidateInputPins(TConstArrayView<FDataLinkPinReference> InInputPins, TConstArrayView<FConstStructView> InInputDataViews) const
+{
+	if (!ensureAlways(InInputPins.Num() == InInputDataViews.Num()))
+	{
+		UE_LOG(LogDataLink, Error, TEXT("[%s] Input Pin count %d does not match Input Data View %d")
+			, GetContextName().GetData()
+			, InInputPins.Num()
+			, InInputDataViews.Num());
+		return false;
+	}
+
+	for (int32 InputIndex = 0; InputIndex < InInputPins.Num(); ++InputIndex)
+	{
+		const FDataLinkPin* InputPin = InInputPins[InputIndex].Pin;
+		check(InputPin);
+
+		const FConstStructView& InputData = InInputDataViews[InputIndex];
+		if (InputPin->Struct && InputPin->Struct != InputData.GetScriptStruct())
+		{
+			UE_LOG(LogDataLink, Error, TEXT("[%s] Input Data struct '%s' does not match expected input pin '%s' of type '%s'")
+				, GetContextName().GetData()
+				, *GetNameSafe(InputData.GetScriptStruct())
+				, *InputPin->Name.ToString()
+				, *InputPin->Struct->GetName());
+			return false;
+		}
+	}
+
+	return true;
+}
+
+bool FDataLinkExecutor::ExecuteEntryNodes()
+{
+	const TArray<FDataLinkPinReference> InputPinsToExecute = Instance.DataLinkGraph->GetInputPins();
+
+	uint16 NodesExecuted = 0;
+	if (!ExecuteInputPins(InputPinsToExecute, TArray<FConstStructView>(Instance.InputData), &NodesExecuted))
+	{
+		return false;
+	}
+
+	// Execute Entry nodes with no input pins (i.e. without dependency) as these were not included in the initial input pins to execute
+	for (const UDataLinkNode* InputNode : Instance.DataLinkGraph->GetInputNodes())
+	{
+		if (InputNode->GetInputPins().IsEmpty())
+		{
+			// Ensure there's a Node Instance to this Input Node
+			FindOrAddNodeInstance(InputNode);
+
+			if (!ExecuteNode(*InputNode))
+			{
+				return false;
+			}
+
+			++NodesExecuted;
+		}
+	}
+
+	return NodesExecuted > 0;
+}
+
+bool FDataLinkExecutor::ExecuteInputPins(TConstArrayView<FDataLinkPinReference> InInputPins, TArray<FConstStructView>&& InInputDataViews, uint16* OutNodesExecuted)
+{
+	if (!ValidateInputPins(InInputPins, InInputDataViews))
+	{
+		return false;
+	}
+
+	// Struct bundling the pin and the data that will be fed into the pin
+	struct FPinData
+	{
+		const FDataLinkPin* Pin;
+		FConstStructView DataView;
+	};
+
+	// Struct bundling the node and the pin data that will be fed into the node
+	struct FNodeData
+	{
+		bool operator==(const UDataLinkNode* InNode) const
+		{
+			return InNode == Node;
+		}
+		const UDataLinkNode* Node = nullptr;
+		TArray<FPinData> PinData;
+	};
+
+	TArray<FNodeData> NodeData;
+
+	// Group the pins per node
+	for (int32 Index = 0; Index < InInputPins.Num(); ++Index)
+	{
+		const FDataLinkPinReference& InputPinReference = InInputPins[Index];
+		check(InputPinReference.Pin);
+
+		FNodeData* NodeDataEntry = NodeData.FindByKey(InputPinReference.OwningNode);
+		if (!NodeDataEntry)
+		{
+			NodeDataEntry = &NodeData.AddDefaulted_GetRef();
+			NodeDataEntry->Node = InputPinReference.OwningNode;
+		}
+
+		FPinData& PinDataEntry = NodeDataEntry->PinData.AddDefaulted_GetRef();
+		PinDataEntry.Pin = InputPinReference.Pin;
+		PinDataEntry.DataView = InInputDataViews[Index];
+	}
+
+	// Set the Node Instance Input Data Views with the corresponding Input Data Views
+	// This does not guarantee that all the input pins have their data filled.
+	// It could be that there's still data not ready yet.
+	for (const FNodeData& NodeDataEntry : NodeData)
+	{
+		FDataLinkInputDataViewer& InputDataViewer = FindOrAddNodeInstance(NodeDataEntry.Node).InputDataViewer;
+
+		// Copy the input data view to the node instance
+		for (const FPinData& PinData : NodeDataEntry.PinData)
+		{
+			InputDataViewer.SetEntryData(*PinData.Pin, PinData.DataView);
+		}
+	}
+
+	uint16 NodesExecuted = 0;
+
+	// Execute nodes that have all their data ready to go
+	for (const FNodeData& NodeInfo : NodeData)
+	{
+		FDataLinkNodeInstance& NodeInstance = GetNodeInstanceMutable(NodeInfo.Node);
+
+		// don't execute nodes that don't have their input data views fully set yet
+		if (NodeInstance.InputDataViewer.HasInvalidDataEntry())
+		{
+			continue;
+		}
+
+		if (!ExecuteNode(*NodeInfo.Node))
+		{
+			return false;
+		}
+
+		++NodesExecuted;
+	}
+
+	if (OutNodesExecuted)
+	{
+		*OutNodesExecuted = NodesExecuted;
+	}
+	return true;
+}
+
+bool FDataLinkExecutor::ExecuteNode(const UDataLinkNode& InNode)
+{
+	FDataLinkNodeInstance& NodeInstance = GetNodeInstanceMutable(&InNode);
+	NodeInstance.Status = EDataLinkNodeStatus::Executing;
+	InNode.Execute(*this);
+	return true;
+}
+
+FDataLinkNodeInstance& FDataLinkExecutor::FindOrAddNodeInstance(const UDataLinkNode* InNode)
+{
+	check(InNode);
+	FDataLinkNodeInstance* NodeInstance = NodeInstanceMap.Find(InNode);
+	if (!NodeInstance)
+	{
+		NodeInstance = &NodeInstanceMap.Add(InNode, FDataLinkNodeInstance(*InNode));
+	}
+	return *NodeInstance;
+}
+
+bool FDataLinkExecutor::ProcessFinishedNode(const UDataLinkNode* InNode)
+{
+	FDataLinkNodeInstance& NodeInstance = GetNodeInstanceMutable(InNode);
+
+	const bool bWasExecuting = ensureAlwaysMsgf(NodeInstance.Status == EDataLinkNodeStatus::Executing
+		, TEXT("[%s] FinishNode called on a node that was not executing!"), GetContextName().GetData());
+
+	NodeInstance.Status = EDataLinkNodeStatus::Finished;
+	return bWasExecuting;
+}
+
+void FDataLinkExecutor::Finish(EDataLinkExecutionResult InStatus, FConstStructView InOutputData)
+{
+	// Finish has already been called
+	if (!bRunning)
+	{
+		return;
+	}
+
+	// Hold onto an extra shared reference to allow implementers to reset their held shared reference in the finish delegate
+	// without destroying this
+	TSharedRef<FDataLinkExecutor> This = SharedThis(this);
+
+	OnFinishDelegate.ExecuteIfBound(*this, InStatus, InOutputData);
+	NodeInstanceMap.Reset();
+	Sink.Reset();
+
+	bRunning = false;
+}
diff --git a/Engine/Plugins/Experimental/DataLink/Source/DataLink/Private/DataLinkExecutorArguments.cpp b/Engine/Plugins/Experimental/DataLink/Source/DataLink/Private/DataLinkExecutorArguments.cpp
new file mode 100644
--- /dev/null
+++ b/Engine/Plugins/Experimental/DataLink/Source/DataLink/Private/DataLinkExecutorArguments.cpp
@@ -0,0 +1,40 @@
+// Copyright Epic Games, Inc. All Rights Reserved.
+
+#include "DataLinkExecutorArguments.h"
+#include "DataLinkExecutor.h"
+
+FDataLinkExecutorArguments::FDataLinkExecutorArguments(const FDataLinkInstance& InInstance)
+	: Instance(InInstance)
+{
+}
+
+FDataLinkExecutorArguments::FDataLinkExecutorArguments(FDataLinkInstance&& InInstance)
+	: Instance(MoveTemp(InInstance))
+{
+}
+
+#if WITH_DATALINK_CONTEXT
+FDataLinkExecutorArguments&& FDataLinkExecutorArguments::SetContextName(FString&& InContextName)
+{
+	ContextName = MoveTemp(InContextName);
+	return MoveTemp(*this);
+}
+#endif
+
+FDataLinkExecutorArguments&& FDataLinkExecutorArguments::SetContextObject(UObject* InContextObject)
+{
+	ContextObject = InContextObject;
+	return MoveTemp(*this);
+}
+
+FDataLinkExecutorArguments&& FDataLinkExecutorArguments::SetSink(const TSharedPtr<FDataLinkSink>& InSink)
+{
+	Sink = InSink;
+	return MoveTemp(*this);
+}
+
+FDataLinkExecutorArguments&& FDataLinkExecutorArguments::SetOnFinish(FOnDataLinkExecutionFinished&& InDelegate)
+{
+	OnFinishDelegate = MoveTemp(InDelegate);
+	return MoveTemp(*this);
+}
diff --git a/Engine/Plugins/Experimental/DataLink/Source/DataLink/Private/DataLinkGraph.cpp b/Engine/Plugins/Experimental/DataLink/Source/DataLink/Private/DataLinkGraph.cpp
new file mode 100644
--- /dev/null
+++ b/Engine/Plugins/Experimental/DataLink/Source/DataLink/Private/DataLinkGraph.cpp
@@ -0,0 +1,56 @@
+// Copyright Epic Games, Inc. All Rights Reserved.
+
+#include "DataLinkGraph.h"
+#include "DataLinkNode.h"
+#include "DataLinkPinReference.h"
+#include "StructUtils/StructView.h"
+
+#if WITH_EDITOR
+TMulticastDelegate<void(UDataLinkGraph*)> UDataLinkGraph::OnGraphCompiledDelegate;
+#endif
+
+int32 UDataLinkGraph::GetInputPinCount() const
+{
+	int32 InputPinCount = 0;
+
+	ForEachInputPin([&InputPinCount](FDataLinkPinReference)
+		{
+			++InputPinCount;
+			return true;
+		});
+
+	return InputPinCount;
+}
+
+TArray<FDataLinkPinReference> UDataLinkGraph::GetInputPins() const
+{
+	TArray<FDataLinkPinReference> InputPins;
+	InputPins.Reserve(GetInputPinCount());
+
+	ForEachInputPin(
+		[&InputPins](FDataLinkPinReference InPinReference)
+		{
+			InputPins.Add(MoveTemp(InPinReference));
+			return true;
+		});
+
+	return InputPins;
+}
+
+bool UDataLinkGraph::ForEachInputPin(TFunctionRef<bool(FDataLinkPinReference)> InFunction) const
+{
+	for (const UDataLinkNode* InputNode : InputNodes)
+	{
+		if (InputNode)
+		{
+			for (const FDataLinkPin& InputPin : InputNode->GetInputPins())
+			{
+				if (InputPin.LinkedNode == nullptr && !InFunction(FDataLinkPinReference(InputNode, &InputPin)))
+				{
+					return false;
+				}
+			}
+		}
+	}
+	return true;
+}
diff --git a/Engine/Plugins/Experimental/DataLink/Source/DataLink/Private/DataLinkInputDataViewer.cpp b/Engine/Plugins/Experimental/DataLink/Source/DataLink/Private/DataLinkInputDataViewer.cpp
new file mode 100644
--- /dev/null
+++ b/Engine/Plugins/Experimental/DataLink/Source/DataLink/Private/DataLinkInputDataViewer.cpp
@@ -0,0 +1,49 @@
+// Copyright Epic Games, Inc. All Rights Reserved.
+
+#include "DataLinkInputDataViewer.h"
+#include "DataLinkPin.h"
+
+FDataLinkInputDataViewer::FDataLinkInputDataViewer(TConstArrayView<FDataLinkPin> InInputPins)
+{
+	DataEntries.Reserve(InInputPins.Num());
+
+	// Initialize data entries with only name set
+	for (const FDataLinkPin& InputPin : InInputPins)
+	{
+		FDataLinkInputDataEntry& DataEntry = DataEntries.AddDefaulted_GetRef();
+		DataEntry.Name = InputPin.Name;
+	}
+}
+
+FConstStructView FDataLinkInputDataViewer::Find(FName InInputName) const
+{
+	if (const FDataLinkInputDataEntry* FoundEntry = DataEntries.FindByKey(InInputName))
+	{
+		return FoundEntry->DataView;
+	}
+	return FConstStructView();
+}
+
+TConstArrayView<FDataLinkInputDataEntry> FDataLinkInputDataViewer::GetDataEntries() const
+{
+	return DataEntries;
+}
+
+bool FDataLinkInputDataViewer::HasInvalidDataEntry() const
+{
+	for (const FDataLinkInputDataEntry& DataEntry : DataEntries)
+	{
+		if (!DataEntry.DataView.IsValid())
+		{
+			return true;
+		}
+	}
+	return false;
+}
+
+void FDataLinkInputDataViewer::SetEntryData(const FDataLinkPin& InPin, FConstStructView InInputDataView)
+{
+	FDataLinkInputDataEntry* DataEntry = DataEntries.FindByKey(InPin.Name);
+	check(DataEntry);
+	DataEntry->DataView = InInputDataView;
+}
diff --git a/Engine/Plugins/Experimental/DataLink/Source/DataLink/Private/DataLinkModule.cpp b/Engine/Plugins/Experimental/DataLink/Source/DataLink/Private/DataLinkModule.cpp
new file mode 100644
--- /dev/null
+++ b/Engine/Plugins/Experimental/DataLink/Source/DataLink/Private/DataLinkModule.cpp
@@ -0,0 +1,8 @@
+// Copyright Epic Games, Inc. All Rights Reserved.
+
+#include "Modules/ModuleManager.h"
+#include "DataLinkLog.h"
+
+DEFINE_LOG_CATEGORY(LogDataLink);
+
+IMPLEMENT_MODULE(FDefaultModuleImpl, DataLink)
diff --git a/Engine/Plugins/Experimental/DataLink/Source/DataLink/Private/DataLinkNode.cpp b/Engine/Plugins/Experimental/DataLink/Source/DataLink/Private/DataLinkNode.cpp
new file mode 100644
--- /dev/null
+++ b/Engine/Plugins/Experimental/DataLink/Source/DataLink/Private/DataLinkNode.cpp
@@ -0,0 +1,37 @@
+// Copyright Epic Games, Inc. All Rights Reserved.
+
+#include "DataLinkNode.h"
+#include "DataLinkEnums.h"
+#include "DataLinkExecutor.h"
+#include "DataLinkLog.h"
+#include "DataLinkNodeMetadata.h"
+#include "DataLinkPinBuilder.h"
+
+void UDataLinkNode::Execute(FDataLinkExecutor& InExecutor) const
+{
+	const EDataLinkExecutionReply Reply = OnExecute(InExecutor);
+
+	if (Reply == EDataLinkExecutionReply::Unhandled)
+	{
+		UE_LOG(LogDataLink, Log, TEXT("[%s] Node '%s' did not handle execution."), InExecutor.GetContextName().GetData(), *GetName())
+		InExecutor.FailNode(this);
+	}
+}
+
+#if WITH_EDITOR
+void UDataLinkNode::BuildMetadata(FDataLinkNodeMetadata& OutMetadata) const
+{
+	OutMetadata
+		.SetDisplayName(GetClass()->GetDisplayNameText())
+		.SetTooltipText(GetClass()->GetToolTipText());
+
+	OnBuildMetadata(OutMetadata);
+}
+#endif
+
+void UDataLinkNode::BuildPins(TArray<FDataLinkPin>& OutInputPins, TArray<FDataLinkPin>& OutOutputPins) const
+{
+	FDataLinkPinBuilder InputBuilder(OutInputPins);
+	FDataLinkPinBuilder OutputBuilder(OutOutputPins);
+	OnBuildPins(InputBuilder, OutputBuilder);
+}
diff --git a/Engine/Plugins/Experimental/DataLink/Source/DataLink/Private/DataLinkNodeInstance.cpp b/Engine/Plugins/Experimental/DataLink/Source/DataLink/Private/DataLinkNodeInstance.cpp
new file mode 100644
--- /dev/null
+++ b/Engine/Plugins/Experimental/DataLink/Source/DataLink/Private/DataLinkNodeInstance.cpp
@@ -0,0 +1,16 @@
+// Copyright Epic Games, Inc. All Rights Reserved.
+
+#include "DataLinkNodeInstance.h"
+#include "DataLinkNode.h"
+
+FDataLinkNodeInstance::FDataLinkNodeInstance(const UDataLinkNode& InNode)
+	: InputDataViewer(InNode.GetInputPins())
+	, OutputDataViewer(InNode.GetOutputPins())
+	, InstanceData(InNode.GetInstanceStruct())
+{
+}
+
+void FDataLinkNodeInstance::AddReferencedObjects(FReferenceCollector& InCollector)
+{
+	InstanceData.AddStructReferencedObjects(InCollector);
+}
diff --git a/Engine/Plugins/Experimental/DataLink/Source/DataLink/Private/DataLinkOutputDataViewer.cpp b/Engine/Plugins/Experimental/DataLink/Source/DataLink/Private/DataLinkOutputDataViewer.cpp
new file mode 100644
--- /dev/null
+++ b/Engine/Plugins/Experimental/DataLink/Source/DataLink/Private/DataLinkOutputDataViewer.cpp
@@ -0,0 +1,34 @@
+// Copyright Epic Games, Inc. All Rights Reserved.
+
+#include "DataLinkOutputDataViewer.h"
+#include "DataLinkPin.h"
+
+FStructView FDataLinkOutputDataEntry::GetDataView(const UScriptStruct* InDesiredStruct) const
+{
+	if (InDesiredStruct && OutputData.GetScriptStruct() != InDesiredStruct)
+	{
+		OutputData.InitializeAs(InDesiredStruct);
+	}
+	return FStructView(OutputData);
+}
+
+FDataLinkOutputDataViewer::FDataLinkOutputDataViewer(TConstArrayView<FDataLinkPin> InOutputPins)
+{
+	DataEntries.Reserve(InOutputPins.Num());
+
+	// Initialize entries with only name set
+	for (const FDataLinkPin& OutputPin : InOutputPins)
+	{
+		FDataLinkOutputDataEntry& DataEntry = DataEntries.AddDefaulted_GetRef();
+		DataEntry.Name = OutputPin.Name;
+	}
+}
+
+FStructView FDataLinkOutputDataViewer::Find(FName InOutputName, const UScriptStruct* InDesiredStruct) const
+{
+	if (const FDataLinkOutputDataEntry* FoundEntry = DataEntries.FindByKey(InOutputName))
+	{
+		return FoundEntry->GetDataView(InDesiredStruct);	
+	}
+	return FStructView();
+}
diff --git a/Engine/Plugins/Experimental/DataLink/Source/DataLink/Private/DataLinkPin.cpp b/Engine/Plugins/Experimental/DataLink/Source/DataLink/Private/DataLinkPin.cpp
new file mode 100644
--- /dev/null
+++ b/Engine/Plugins/Experimental/DataLink/Source/DataLink/Private/DataLinkPin.cpp
@@ -0,0 +1,34 @@
+// Copyright Epic Games, Inc. All Rights Reserved.
+
+#include "DataLinkPin.h"
+#include "DataLinkNode.h"
+
+FDataLinkPin::FDataLinkPin(FName InPinName)
+	: Name(InPinName)
+{
+}
+
+const FDataLinkPin* FDataLinkPin::GetLinkedInputPin() const
+{
+	if (!LinkedNode)
+	{
+		return nullptr;
+	}
+
+	TConstArrayView<FDataLinkPin> LinkedInputPins = LinkedNode->GetInputPins();
+	if (LinkedInputPins.IsValidIndex(LinkedIndex))
+	{
+		return &LinkedInputPins[LinkedIndex];
+	}
+
+	return nullptr;
+}
+
+FText FDataLinkPin::GetDisplayName() const
+{
+	if (DisplayName.IsEmpty())
+	{
+		return FText::FromString(Name.ToString());
+	}
+	return DisplayName;
+}
diff --git a/Engine/Plugins/Experimental/DataLink/Source/DataLink/Private/DataLinkPinBuilder.cpp b/Engine/Plugins/Experimental/DataLink/Source/DataLink/Private/DataLinkPinBuilder.cpp
new file mode 100644
--- /dev/null
+++ b/Engine/Plugins/Experimental/DataLink/Source/DataLink/Private/DataLinkPinBuilder.cpp
@@ -0,0 +1,36 @@
+// Copyright Epic Games, Inc. All Rights Reserved.
+
+#include "DataLinkPinBuilder.h"
+
+FDataLinkPinParameters::FDataLinkPinParameters(FDataLinkPin& InPin)
+	: Pin(InPin)
+{
+}
+
+FDataLinkPinParameters& FDataLinkPinParameters::SetDisplayName(FText&& InPinDisplayName)
+{
+	Pin.DisplayName = MoveTemp(InPinDisplayName);
+	return *this;
+}
+
+FDataLinkPinParameters& FDataLinkPinParameters::SetStruct(const UScriptStruct* InPinStruct)
+{
+	Pin.Struct = InPinStruct;
+	return *this;
+}
+
+FDataLinkPinBuilder::FDataLinkPinBuilder(TArray<FDataLinkPin>& OutPins)
+	: Pins(OutPins)
+{
+}
+
+void FDataLinkPinBuilder::AddCapacity(int32 InNumToAdd)
+{
+	Pins.Reserve(Pins.Num() + InNumToAdd);
+}
+
+FDataLinkPinParameters FDataLinkPinBuilder::Add(FName InPinName)
+{
+	checkf(!Pins.Contains(InPinName), TEXT("Pin Name '%s' already exists! Should be unique"), *InPinName.ToString());
+	return FDataLinkPinParameters(Pins.Emplace_GetRef(InPinName));
+}
diff --git a/Engine/Plugins/Experimental/DataLink/Source/DataLink/Private/DataLinkRequestProxy.cpp b/Engine/Plugins/Experimental/DataLink/Source/DataLink/Private/DataLinkRequestProxy.cpp
new file mode 100644
--- /dev/null
+++ b/Engine/Plugins/Experimental/DataLink/Source/DataLink/Private/DataLinkRequestProxy.cpp
@@ -0,0 +1,51 @@
+// Copyright Epic Games, Inc. All Rights Reserved.
+
+#include "DataLinkRequestProxy.h"
+#include "Containers/Ticker.h"
+#include "DataLinkExecutor.h"
+#include "DataLinkExecutorArguments.h"
+#include "DataLinkUtils.h"
+
+UDataLinkRequestProxy* UDataLinkRequestProxy::CreateRequestProxy(FDataLinkInstance InDataLinkInstance, UObject* InExecutionContext, TScriptInterface<IDataLinkSinkProvider> InDataLinkSinkProvider)
+{
+	UDataLinkRequestProxy* RequestProxy = NewObject<UDataLinkRequestProxy>();
+	RequestProxy->SetFlags(RF_StrongRefOnFrame);
+	RequestProxy->ProcessRequest(MoveTemp(InDataLinkInstance), InExecutionContext, InDataLinkSinkProvider);
+	return RequestProxy;
+}
+
+void UDataLinkRequestProxy::ProcessRequest(FDataLinkInstance&& InDataLinkInstance, UObject* InExecutionContext, TScriptInterface<IDataLinkSinkProvider> InDataLinkSinkProvider)
+{
+	DataLinkExecutor = FDataLinkExecutor::Create(FDataLinkExecutorArguments(MoveTemp(InDataLinkInstance))
+#if WITH_DATALINK_CONTEXT
+		.SetContextName(GetName())
+#endif
+		.SetContextObject(InExecutionContext)
+		.SetSink(UE::DataLink::TryGetSink(InDataLinkSinkProvider))
+		.SetOnFinish(FOnDataLinkExecutionFinished::CreateUObject(this, &UDataLinkRequestProxy::OnExecutionFinished)));
+
+	// Delay execution by 1 frame. This is so that even for executions that finish immediately,
+	// the Async Call blueprint script still has chance to get the returned Proxy and listen to execution finished event
+	// Note: a fix for this would be to have the Async Node create the object first, bind to the delegate, then execute a 'run' command
+	FTSTicker::GetCoreTicker().AddTicker(FTickerDelegate::CreateWeakLambda(this,
+		[this](float) mutable -> bool 
+		{
+			if (ensureAlways(DataLinkExecutor.IsValid()))
+			{
+				DataLinkExecutor->Run();
+			}
+			// Fire once
+			return false;
+		}));
+}
+
+void UDataLinkRequestProxy::OnExecutionFinished(const FDataLinkExecutor& InExecutor, EDataLinkExecutionResult InResult, FConstStructView InOutputData)
+{
+	DataLinkExecutor.Reset();
+
+	FInstancedStruct OutputData(InOutputData);
+
+	// Allow Blueprint CallInEditor functions to handle then execution result
+	FEditorScriptExecutionGuard ScriptGuard;
+	OnRequestComplete.Broadcast(OutputData, InResult);
+}
diff --git a/Engine/Plugins/Experimental/DataLink/Source/DataLink/Private/DataLinkSink.cpp b/Engine/Plugins/Experimental/DataLink/Source/DataLink/Private/DataLinkSink.cpp
new file mode 100644
--- /dev/null
+++ b/Engine/Plugins/Experimental/DataLink/Source/DataLink/Private/DataLinkSink.cpp
@@ -0,0 +1,98 @@
+// Copyright Epic Games, Inc. All Rights Reserved.
+
+#include "DataLinkSink.h"
+#include "DataLinkNode.h"
+#include "DataLinkNodeInstance.h"
+#include "StructUtils/UserDefinedStruct.h"
+
+namespace UE::DataLink::Private
+{
+	bool CompareDataEntries(const FDataLinkInputDataEntry& InDataEntry, const FDataLinkInputDataEntry& InOtherDataEntry)
+	{
+		if (InDataEntry.Name != InOtherDataEntry.Name)
+		{
+			return false;
+		}
+
+		const UScriptStruct* DataStruct = InDataEntry.DataView.GetScriptStruct();
+		if (DataStruct != InOtherDataEntry.DataView.GetScriptStruct())
+		{
+			return false;
+		}
+
+		constexpr uint32 PortFlags = PPF_None;
+		return !DataStruct || DataStruct->CompareScriptStruct(InDataEntry.DataView.GetMemory(), InOtherDataEntry.DataView.GetMemory(), PortFlags);
+	}
+}
+
+FDataLinkSinkKey::FDataLinkSinkKey(TSubclassOf<UDataLinkNode> InNodeClass, const FDataLinkInputDataViewer& InInputDataViewer)
+	: NodeClass(InNodeClass)
+	, InputDataEntries(InInputDataViewer.GetDataEntries())
+{
+	CachedHash = GetTypeHash(NodeClass);
+
+	for (const FDataLinkInputDataEntry& InputDataEntry : InputDataEntries)
+	{
+		const UScriptStruct* InputStruct = InputDataEntry.DataView.GetScriptStruct();
+
+		CachedHash = HashCombineFast(CachedHash, GetTypeHash(InputDataEntry.Name));
+		CachedHash = HashCombineFast(CachedHash, GetTypeHash(InputStruct));
+
+		// GetStructHash64 does not check for whether CppStructOps is valid or not
+		if (InputStruct && InputStruct->GetCppStructOps())
+		{
+			CachedHash = HashCombineFast(CachedHash, UE::StructUtils::GetStructHash64(InputDataEntry.DataView));
+		}
+	}
+}
+
+bool FDataLinkSinkKey::operator==(const FDataLinkSinkKey& InOther) const
+{
+	if (CachedHash != InOther.CachedHash)
+	{
+		return false;
+	}
+
+	if (NodeClass != InOther.NodeClass)
+	{
+		return false;
+	}
+
+	if (InputDataEntries.Num() != InOther.InputDataEntries.Num())
+	{
+		return false;
+	}
+
+	for (int32 Index = 0; Index < InputDataEntries.Num(); ++Index)
+	{
+		if (!UE::DataLink::Private::CompareDataEntries(InputDataEntries[Index], InOther.InputDataEntries[Index]))
+		{
+			return false;
+		}
+	}
+
+	return true;
+}
+
+uint64 GetTypeHash(const FDataLinkSinkKey& InKey)
+{
+	return InKey.CachedHash;
+}
+
+FInstancedStruct& FDataLinkSink::FindOrAddCachedData(const FDataLinkNodeInstance& InNodeInstance, TOptional<const UScriptStruct*> InDesiredStruct)
+{
+	FInstancedStruct& InstancedStruct = CachedDataMap.FindOrAdd(InNodeInstance.GetSinkKey());
+	if (InDesiredStruct.IsSet() && InstancedStruct.GetScriptStruct() != *InDesiredStruct)
+	{
+		InstancedStruct.InitializeAs(*InDesiredStruct);
+	}
+	return InstancedStruct;
+}
+
+void FDataLinkSink::AddStructReferencedObjects(FReferenceCollector& InCollector)
+{
+	for (TPair<FDataLinkSinkKey, FInstancedStruct>& Pair : CachedDataMap)
+	{
+		Pair.Value.AddStructReferencedObjects(InCollector);
+	}
+}
diff --git a/Engine/Plugins/Experimental/DataLink/Source/DataLink/Private/DataLinkSinkObject.cpp b/Engine/Plugins/Experimental/DataLink/Source/DataLink/Private/DataLinkSinkObject.cpp
new file mode 100644
--- /dev/null
+++ b/Engine/Plugins/Experimental/DataLink/Source/DataLink/Private/DataLinkSinkObject.cpp
@@ -0,0 +1,24 @@
+// Copyright Epic Games, Inc. All Rights Reserved.
+
+#include "DataLinkSinkObject.h"
+#include "DataLinkSink.h"
+
+UDataLinkSinkObject::UDataLinkSinkObject()
+{
+	Sink = MakeShared<FDataLinkSink>();
+}
+
+void UDataLinkSinkObject::ResetSink()
+{
+	Sink = MakeShared<FDataLinkSink>();
+}
+
+TSharedPtr<FDataLinkSink> UDataLinkSinkObject::GetSink() const
+{
+	return Sink;
+}
+
+const UDataLinkSinkObject* UDataLinkSinkObject::GetSinkObject_Implementation() const
+{
+	return this;
+}
diff --git a/Engine/Plugins/Experimental/DataLink/Source/DataLink/Private/DataLinkUtils.cpp b/Engine/Plugins/Experimental/DataLink/Source/DataLink/Private/DataLinkUtils.cpp
new file mode 100644
--- /dev/null
+++ b/Engine/Plugins/Experimental/DataLink/Source/DataLink/Private/DataLinkUtils.cpp
@@ -0,0 +1,84 @@
+// Copyright Epic Games, Inc. All Rights Reserved.
+
+#include "DataLinkUtils.h"
+#include "DataLinkLog.h"
+#include "DataLinkSinkObject.h"
+#include "Engine/Engine.h"
+#include "IDataLinkSinkProvider.h"
+#include "UObject/Class.h"
+#include "UObject/Object.h"
+#include "UObject/Package.h"
+
+bool UE::DataLink::ReplaceObject(UObject*& InOutObject, UObject* InOuter, UClass* InClass)
+{
+	// Check if the object class already matches the new class
+	if (InOutObject && InOutObject->GetClass() == InClass)
+	{
+		UE_LOG(LogDataLink, Log, TEXT("ReplaceObject did not take place as '%s' is already of class %s.")
+			, *InOutObject->GetName()
+			, *GetNameSafe(InClass));
+		return false;
+	}
+
+	UObject* const OldObject = InOutObject;
+	InOutObject = nullptr;
+
+	// Save the current object name before renaming it
+	const FName ObjectName = InOutObject ? InOutObject->GetFName() : NAME_None;
+
+	// Discard current object
+	if (InOutObject)
+	{
+		UObject* const NewOuter = GetTransientPackage();
+		const FName UniqueName = MakeUniqueObjectName(NewOuter, InOutObject->GetClass(), *(TEXT("TRASH_") + InOutObject->GetName()));
+		InOutObject->Rename(*UniqueName.ToString(), NewOuter, REN_DoNotDirty | REN_DontCreateRedirectors | REN_NonTransactional);
+		InOutObject->MarkAsGarbage();
+		InOutObject = nullptr;
+	}
+
+	// Set the new class (only if new class is valid)
+	// The operation is still considered valid if InClass is null. This 
+	if (InClass)
+	{
+		const EObjectFlags ObjectFlags = InOuter ? InOuter->GetMaskedFlags(RF_PropagateToSubObjects) : RF_NoFlags;
+
+		InOutObject = NewObject<UObject>(InOuter, InClass, ObjectName, ObjectFlags);
+
+		if (OldObject && GEngine)
+		{
+			TMap<UObject*, UObject*> ReplacementMap;
+			ReplacementMap.Add(OldObject, InOutObject);
+			GEngine->NotifyToolsOfObjectReplacement(ReplacementMap);
+		}
+	}
+
+	return true;
+}
+
+TSharedPtr<FDataLinkSink> UE::DataLink::TryGetSink(TScriptInterface<IDataLinkSinkProvider> InSinkProvider)
+{
+	if (IDataLinkSinkProvider* SinkProvider = InSinkProvider.GetInterface())
+	{
+		if (TSharedPtr<FDataLinkSink> Sink = SinkProvider->GetSink())
+		{
+			return Sink;
+		}
+
+		if (const UDataLinkSinkObject* SinkObject = SinkProvider->GetSinkObject())
+		{
+			return SinkObject->GetSink();
+		}
+	}
+	else if (UObject* Object = InSinkProvider.GetObject())
+	{
+		if (Object->GetClass()->ImplementsInterface(UDataLinkSinkProvider::StaticClass()))
+		{
+			if (const UDataLinkSinkObject* SinkObject = IDataLinkSinkProvider::Execute_GetSinkObject(Object))
+			{
+				return SinkObject->GetSink();
+			}
+		}
+	}
+
+	return nullptr;
+}
diff --git a/Engine/Plugins/Experimental/DataLink/Source/DataLink/Private/Nodes/DataLinkConstant.cpp b/Engine/Plugins/Experimental/DataLink/Source/DataLink/Private/Nodes/DataLinkConstant.cpp
new file mode 100644
--- /dev/null
+++ b/Engine/Plugins/Experimental/DataLink/Source/DataLink/Private/Nodes/DataLinkConstant.cpp
@@ -0,0 +1,49 @@
+// Copyright Epic Games, Inc. All Rights Reserved.
+
+#include "Nodes/DataLinkConstant.h"
+#include "DataLinkExecutor.h"
+#include "DataLinkNames.h"
+#include "DataLinkNodeInstance.h"
+#include "DataLinkNodeMetadata.h"
+#include "DataLinkPinBuilder.h"
+
+void UDataLinkConstant::SetStruct(const UScriptStruct* InStructType)
+{
+	Instance.InitializeAs(InStructType);
+}
+
+#if WITH_EDITOR
+void UDataLinkConstant::OnBuildMetadata(FDataLinkNodeMetadata& Metadata) const
+{
+	Super::OnBuildMetadata(Metadata);
+
+	if (const UScriptStruct* Struct = Instance.GetScriptStruct())
+	{
+		Metadata
+			.SetDisplayName(Struct->GetDisplayNameText())
+			.SetTooltipText(Struct->GetToolTipText());
+	}
+
+	if (!DisplayName.IsEmpty())
+	{
+		Metadata.SetDisplayName(DisplayName);
+	}
+}
+#endif
+
+void UDataLinkConstant::OnBuildPins(FDataLinkPinBuilder& Inputs, FDataLinkPinBuilder& Outputs) const
+{
+	Super::OnBuildPins(Inputs, Outputs);
+
+	if (const UScriptStruct* Struct = Instance.GetScriptStruct())
+	{
+		Outputs.Add(UE::DataLink::OutputDefault)
+			.SetStruct(Struct);
+	}
+}
+
+EDataLinkExecutionReply UDataLinkConstant::OnExecute(FDataLinkExecutor& InExecutor) const
+{
+	InExecutor.SucceedNode(this, FConstStructView(Instance.GetScriptStruct(), Instance.GetMemory()));
+	return EDataLinkExecutionReply::Handled;
+}
diff --git a/Engine/Plugins/Experimental/DataLink/Source/DataLink/Private/Nodes/Script/DataLinkScriptNode.cpp b/Engine/Plugins/Experimental/DataLink/Source/DataLink/Private/Nodes/Script/DataLinkScriptNode.cpp
new file mode 100644
--- /dev/null
+++ b/Engine/Plugins/Experimental/DataLink/Source/DataLink/Private/Nodes/Script/DataLinkScriptNode.cpp
@@ -0,0 +1,181 @@
+// Copyright Epic Games, Inc. All Rights Reserved.
+
+#include "Nodes/Script/DataLinkScriptNode.h"
+#include "Blueprint/BlueprintExceptionInfo.h"
+#include "DataLinkCoreTypes.h"
+#include "DataLinkExecutor.h"
+
+#define LOCTEXT_NAMESPACE "DataLinkScriptNode"
+
+void UDataLinkScriptNode::Execute(const UDataLinkNode* InNode, FDataLinkExecutor& InExecutor)
+{
+	Node = InNode;
+	ExecutorWeak = InExecutor.AsWeak();
+	OnExecute();
+}
+
+bool UDataLinkScriptNode::Succeed(const FInstancedStruct& OutputData)
+{
+	const TSharedPtr<FDataLinkExecutor> Executor = ExecutorWeak.Pin();
+	if (!Executor.IsValid())
+	{
+		return false;
+	}
+
+	const UScriptStruct* OutputDataStruct = OutputData.GetScriptStruct();
+	if (!OutputDataStruct || OutputPin.Struct != OutputDataStruct)
+	{
+		Executor->FailNode(Node);
+		return false;
+	}
+
+	const FDataLinkNodeInstance& NodeInstance = Executor->GetNodeInstance(Node);
+
+	const FDataLinkOutputDataViewer& OutputDataViewer = NodeInstance.GetOutputDataViewer();
+
+	const FStructView OutputDataView = OutputDataViewer.Find(OutputPin.Name, OutputPin.Struct);
+	if (!OutputDataView.IsValid())
+	{
+		Executor->FailNode(Node);
+		return false;
+	}
+
+	OutputDataStruct->CopyScriptStruct(OutputDataView.GetMemory(), OutputData.GetMemory());
+	Executor->SucceedNode(Node, OutputDataView);
+	return true;
+}
+
+void UDataLinkScriptNode::Fail()
+{
+	if (TSharedPtr<FDataLinkExecutor> Executor = ExecutorWeak.Pin())
+	{
+		Executor->FailNode(Node);
+	}
+}
+
+bool UDataLinkScriptNode::GetInputData(FInstancedStruct& InputData, FName InputName) const
+{
+	if (TSharedPtr<FDataLinkExecutor> Executor = ExecutorWeak.Pin())
+	{
+		const FDataLinkInputDataViewer& InputDataViewer = Executor->GetNodeInstance(Node).GetInputDataViewer();
+		const FConstStructView InputDataView = InputDataViewer.Find(InputName);
+		if (InputDataView.IsValid())
+		{
+			InputData = FInstancedStruct(InputDataView);
+			return true;
+		}
+	}
+	return false;
+}
+
+UWorld* UDataLinkScriptNode::GetWorld() const
+{
+	if (UObject* Context = GetContextObject())
+	{
+		return Context->GetWorld();
+	}
+	return nullptr;
+}
+
+UObject* UDataLinkScriptNode::GetContextObject() const
+{
+	if (TSharedPtr<FDataLinkExecutor> Executor = ExecutorWeak.Pin())
+	{
+		return Executor->GetContextObject();
+	}
+	return nullptr;
+}
+
+DEFINE_FUNCTION(UDataLinkScriptNode::execSucceedWildcard)
+{
+	// Read Out Struct property
+	Stack.MostRecentPropertyAddress = nullptr;
+	Stack.MostRecentPropertyContainer = nullptr;
+	Stack.StepCompiledIn<FStructProperty>(nullptr);
+
+	const FStructProperty* OutputStructProperty = CastField<FStructProperty>(Stack.MostRecentProperty);
+	void* OutputStructAddress = Stack.MostRecentPropertyAddress;
+
+	P_FINISH;
+
+	*(bool*)RESULT_PARAM = false;
+
+	if (!OutputStructProperty || !OutputStructProperty->Struct || !OutputStructAddress)
+	{
+		FBlueprintCoreDelegates::ThrowScriptException(P_THIS, Stack, FBlueprintExceptionInfo(EBlueprintExceptionType::AbortExecution
+			, LOCTEXT("InvalidOutputStructWarning", "Failed to resolve the Output Data on 'Finish Event'")));
+
+		P_THIS->Fail();
+		return;
+	}
+
+	if (P_THIS->OutputPin.Struct != OutputStructProperty->Struct)
+	{
+		P_THIS->Fail();
+		return;
+	}
+
+	P_NATIVE_BEGIN;
+
+	if (TSharedPtr<FDataLinkExecutor> Executor = P_THIS->ExecutorWeak.Pin())
+	{
+		const FDataLinkNodeInstance& NodeInstance = Executor->GetNodeInstance(P_THIS->Node);
+
+		const FDataLinkOutputDataViewer& OutputDataViewer = NodeInstance.GetOutputDataViewer();
+
+		const FStructView OutputDataView = NodeInstance.GetOutputDataViewer().Find(P_THIS->OutputPin.Name, P_THIS->OutputPin.Struct);
+		if (OutputDataView.IsValid())
+		{
+			OutputStructProperty->Struct->CopyScriptStruct(OutputDataView.GetMemory(), OutputStructAddress);
+			Executor->SucceedNode(P_THIS->Node, OutputDataView);
+			*(bool*)RESULT_PARAM = true;
+		}
+		else
+		{
+			Executor->FailNode(P_THIS->Node);
+		}
+	}
+
+	P_NATIVE_END;
+}
+
+DEFINE_FUNCTION(UDataLinkScriptNode::execGetInputDataWildcard)
+{
+	Stack.MostRecentPropertyAddress = nullptr;
+	Stack.MostRecentPropertyContainer = nullptr;
+	Stack.StepCompiledIn<FStructProperty>(nullptr);
+
+	const FStructProperty* InputStructProperty = CastField<FStructProperty>(Stack.MostRecentProperty);
+	void* InputStructAddress = Stack.MostRecentPropertyAddress;
+
+	P_GET_PROPERTY(FNameProperty, InputName);
+
+	P_FINISH;
+
+	*(bool*)RESULT_PARAM = false;
+
+	if (!InputStructProperty || !InputStructProperty->Struct || !InputStructAddress)
+	{
+		FBlueprintCoreDelegates::ThrowScriptException(P_THIS, Stack, FBlueprintExceptionInfo(EBlueprintExceptionType::AbortExecution
+			, LOCTEXT("InvalidInputStructWarning", "Failed to resolve the Input Data on 'Get Input Data'")));
+		return;
+	}
+
+	P_NATIVE_BEGIN;
+
+	if (TSharedPtr<FDataLinkExecutor> Executor = P_THIS->ExecutorWeak.Pin())
+	{
+		const FDataLinkInputDataViewer& InputDataViewer = Executor->GetNodeInstance(P_THIS->Node).GetInputDataViewer();
+		const FConstStructView InputDataView = InputDataViewer.Find(InputName);
+
+		if (InputDataView.IsValid() && InputDataView.GetScriptStruct() == InputStructProperty->Struct)
+		{
+			InputStructProperty->Struct->CopyScriptStruct(InputStructAddress, InputDataView.GetMemory());
+			*(bool*)RESULT_PARAM = true;
+		}
+	}
+
+	P_NATIVE_END;
+}
+
+#undef LOCTEXT_NAMESPACE
diff --git a/Engine/Plugins/Experimental/DataLink/Source/DataLink/Private/Nodes/Script/DataLinkScriptNodeWrapper.cpp b/Engine/Plugins/Experimental/DataLink/Source/DataLink/Private/Nodes/Script/DataLinkScriptNodeWrapper.cpp
new file mode 100644
--- /dev/null
+++ b/Engine/Plugins/Experimental/DataLink/Source/DataLink/Private/Nodes/Script/DataLinkScriptNodeWrapper.cpp
@@ -0,0 +1,85 @@
+// Copyright Epic Games, Inc. All Rights Reserved.
+
+#include "Nodes/Script/DataLinkScriptNodeWrapper.h"
+#include "DataLinkExecutor.h"
+#include "DataLinkLog.h"
+#include "DataLinkNodeInstance.h"
+#include "DataLinkNodeMetadata.h"
+#include "DataLinkPinBuilder.h"
+#include "DataLinkUtils.h"
+#include "Nodes/Script/DataLinkScriptNode.h"
+#include "UObject/Package.h"
+
+UDataLinkScriptNodeWrapper::UDataLinkScriptNodeWrapper()
+{
+	InstanceStruct = FDataLinkScriptNodeInstance::StaticStruct();
+}
+
+void UDataLinkScriptNodeWrapper::SetNodeClass(TSubclassOf<UDataLinkScriptNode> InNodeClass)
+{
+	NodeClass = InNodeClass;
+
+	UObject* Node = TemplateNode;
+	if (UE::DataLink::ReplaceObject(Node, this, InNodeClass))
+	{
+		TemplateNode = Cast<UDataLinkScriptNode>(Node);
+	}
+}
+
+#if WITH_EDITOR
+void UDataLinkScriptNodeWrapper::OnBuildMetadata(FDataLinkNodeMetadata& Metadata) const
+{
+	Super::OnBuildMetadata(Metadata);
+
+	if (NodeClass)
+	{
+		Metadata
+			.SetDisplayName(NodeClass->GetDisplayNameText())
+			.SetTooltipText(NodeClass->GetToolTipText());
+	}
+}
+#endif
+
+void UDataLinkScriptNodeWrapper::OnBuildPins(FDataLinkPinBuilder& Inputs, FDataLinkPinBuilder& Outputs) const
+{
+	Super::OnBuildPins(Inputs, Outputs);
+
+	if (TemplateNode)
+	{
+		Inputs.AddCapacity(TemplateNode->InputPins.Num());
+
+		for (const FDataLinkScriptPin& Pin : TemplateNode->InputPins)
+		{
+			Inputs.Add(Pin.Name)
+				.SetStruct(Pin.Struct);
+		}
+
+		Outputs.Add(TemplateNode->OutputPin.Name)
+			.SetStruct(TemplateNode->OutputPin.Struct);
+	}
+}
+
+EDataLinkExecutionReply UDataLinkScriptNodeWrapper::OnExecute(FDataLinkExecutor& InExecutor) const
+{
+	if (!NodeClass)
+	{
+		return EDataLinkExecutionReply::Unhandled;
+	}
+
+	FDataLinkNodeInstance& NodeInstance = InExecutor.GetNodeInstanceMutable(this);
+
+	FDataLinkScriptNodeInstance& Instance = NodeInstance.GetInstanceDataMutable().Get<FDataLinkScriptNodeInstance>();
+	Instance.Node = NewObject<UDataLinkScriptNode>(GetTransientPackage(), NodeClass, NAME_None, RF_Transient, TemplateNode);
+	Instance.Node->Execute(this, InExecutor);
+
+	return EDataLinkExecutionReply::Handled;
+}
+
+UDataLinkScriptNode* UDataLinkScriptNodeWrapper::GetDefaultNode() const
+{
+	if (NodeClass)
+	{
+		return NodeClass->GetDefaultObject<UDataLinkScriptNode>();
+	}
+	return nullptr;
+}
diff --git a/Engine/Plugins/Experimental/DataLink/Source/DataLink/Private/Nodes/String/DataLinkNodeStringBuilder.cpp b/Engine/Plugins/Experimental/DataLink/Source/DataLink/Private/Nodes/String/DataLinkNodeStringBuilder.cpp
new file mode 100644
--- /dev/null
+++ b/Engine/Plugins/Experimental/DataLink/Source/DataLink/Private/Nodes/String/DataLinkNodeStringBuilder.cpp
@@ -0,0 +1,50 @@
+// Copyright Epic Games, Inc. All Rights Reserved.
+
+#include "Nodes/String/DataLinkNodeStringBuilder.h"
+#include "DataLinkCoreTypes.h"
+#include "DataLinkExecutor.h"
+#include "DataLinkNames.h"
+#include "DataLinkPinBuilder.h"
+#include "Nodes/String/DataLinkStringBuilder.h"
+
+void UDataLinkNodeStringBuilder::OnBuildPins(FDataLinkPinBuilder& Inputs, FDataLinkPinBuilder& Outputs) const
+{
+	FDataLinkStringBuilder(Segments, Tokens)
+		.BuildInputPins(Inputs);
+
+	Outputs.Add(UE::DataLink::OutputDefault)
+		.SetStruct<FDataLinkString>();
+}
+
+EDataLinkExecutionReply UDataLinkNodeStringBuilder::OnExecute(FDataLinkExecutor& InExecutor) const
+{
+	const FDataLinkNodeInstance& NodeInstance = InExecutor.GetNodeInstance(this);
+
+	const FDataLinkOutputDataViewer& OutputDataViewer = NodeInstance.GetOutputDataViewer();
+
+	FDataLinkString& OutputString = OutputDataViewer.Get<FDataLinkString>(UE::DataLink::OutputDefault);
+
+	const FDataLinkStringBuilder StringBuilder(Segments, Tokens);
+	if (StringBuilder.BuildString(NodeInstance.GetInputDataViewer(), OutputString.Value))
+	{
+		InExecutor.SucceedNode(this, FConstStructView::Make(OutputString));
+	}
+	else
+	{
+		InExecutor.FailNode(this);
+	}
+
+	return EDataLinkExecutionReply::Handled;
+}
+
+#if WITH_EDITOR
+void UDataLinkNodeStringBuilder::PostEditChangeProperty(FPropertyChangedEvent& InPropertyChangedEvent)
+{
+	Super::PostEditChangeProperty(InPropertyChangedEvent);
+
+	if (InPropertyChangedEvent.GetMemberPropertyName() == GET_MEMBER_NAME_CHECKED(UDataLinkNodeStringBuilder, Segments))
+	{
+		FDataLinkStringBuilder::GatherTokens(Segments, Tokens);
+	}
+}
+#endif
diff --git a/Engine/Plugins/Experimental/DataLink/Source/DataLink/Private/Nodes/String/DataLinkStringBuilder.cpp b/Engine/Plugins/Experimental/DataLink/Source/DataLink/Private/Nodes/String/DataLinkStringBuilder.cpp
new file mode 100644
--- /dev/null
+++ b/Engine/Plugins/Experimental/DataLink/Source/DataLink/Private/Nodes/String/DataLinkStringBuilder.cpp
@@ -0,0 +1,117 @@
+// Copyright Epic Games, Inc. All Rights Reserved.
+
+#include "Nodes/String/DataLinkStringBuilder.h"
+#include "DataLinkCoreTypes.h"
+#include "DataLinkInputDataViewer.h"
+#include "DataLinkPinBuilder.h"
+#include "StructUtils/StructView.h"
+
+namespace UE::DataLink::Private
+{
+	FStringView TryGetTokenName(const FStringView InString)
+	{
+		// String with just {} should not be valid. Should at least have something in between the brackets
+		if (InString.Len() > 2
+			&& InString.StartsWith(TEXT("{"))
+			&& InString.EndsWith(TEXT("}")))
+		{
+			// Removes the { and } from the start and end
+			return InString.Mid(1, InString.Len() - 2);
+		}
+		return FStringView();
+	}
+}
+
+bool FDataLinkStringBuilder::BuildString(const FDataLinkInputDataViewer& InTokenValues, FString& InOutResult) const
+{
+	if (InTokenValues.Num() != Tokens.Num())
+	{
+		return false;
+	}
+
+	int32 ResultStringLength = 0;
+	ForEachResolvedSegment(InTokenValues,
+		[&ResultStringLength](FStringView InResolvedSegment)
+		{
+			ResultStringLength += InResolvedSegment.Len();
+		});
+
+	InOutResult.Reset(ResultStringLength);
+
+	ForEachResolvedSegment(InTokenValues,
+		[&InOutResult](FStringView InResolvedSegment)
+		{
+			InOutResult += InResolvedSegment;
+		});
+
+	return true;
+}
+
+void FDataLinkStringBuilder::BuildInputPins(FDataLinkPinBuilder& Inputs) const
+{
+	Inputs.AddCapacity(Tokens.Num());
+
+	for (const FDataLinkStringBuilderToken& Token : Tokens)
+	{
+		Inputs.Add(Token.Name)
+			.SetStruct<FDataLinkString>();
+	}
+}
+
+void FDataLinkStringBuilder::GatherTokens(TConstArrayView<FString> InSegments, TArray<FDataLinkStringBuilderToken>& OutTokens)
+{
+	// Try to re-use the old token count if existing as an initial allocation capacity
+	TMap<FStringView, int32> TokenIndexMap;
+	TokenIndexMap.Reserve(OutTokens.Num());
+
+	OutTokens.Reset();
+
+	for (int32 SegmentIndex = 0; SegmentIndex < InSegments.Num(); ++SegmentIndex)
+	{
+		const FStringView TokenName = UE::DataLink::Private::TryGetTokenName(InSegments[SegmentIndex]);
+		if (TokenName.IsEmpty())
+		{
+			continue;
+		}
+
+		if (int32* ExistingTokenIndex = TokenIndexMap.Find(TokenName))
+		{
+			FDataLinkStringBuilderToken& Token = OutTokens[*ExistingTokenIndex];
+			checkSlow(Token.Name == TokenName);
+			Token.Indices.Add(SegmentIndex);
+		}
+		else
+		{
+			const int32 TokenIndex = OutTokens.AddDefaulted();
+			TokenIndexMap.Add(TokenName, TokenIndex);
+
+			FDataLinkStringBuilderToken& Token = OutTokens[TokenIndex];
+			Token.Name = FName(TokenName);
+			Token.Indices.Add(SegmentIndex);
+		}
+	}
+}
+
+void FDataLinkStringBuilder::ForEachResolvedSegment(const FDataLinkInputDataViewer& InTokenValues, TFunctionRef<void(FStringView InResolvedSegment)> InFunction) const
+{
+	// At this point, these two are required to match in count
+	check(Tokens.Num() == InTokenValues.Num());
+
+	for (int32 SegmentIndex = 0; SegmentIndex < Segments.Num(); ++SegmentIndex)
+	{
+		const FDataLinkStringBuilderToken* Token = Tokens.FindByPredicate(
+			[SegmentIndex](const FDataLinkStringBuilderToken& InToken)->bool
+			{
+				return InToken.Indices.Contains(SegmentIndex);
+			});
+
+		if (Token)
+		{
+			InFunction(InTokenValues.Get<FDataLinkString>(Token->Name).Value);
+		}
+		else
+		{
+			InFunction(Segments[SegmentIndex]);
+		}
+	}
+}
diff --git a/Engine/Plugins/Experimental/DataLink/Source/DataLink/Public/DataLinkCoreTypes.h b/Engine/Plugins/Experimental/DataLink/Source/DataLink/Public/DataLinkCoreTypes.h
new file mode 100644
--- /dev/null
+++ b/Engine/Plugins/Experimental/DataLink/Source/DataLink/Public/DataLinkCoreTypes.h
@@ -0,0 +1,34 @@
+// Copyright Epic Games, Inc. All Rights Reserved.
+
+#pragma once
+
+#include "Containers/UnrealString.h"
+#include "DataLinkCoreTypes.generated.h"
+
+USTRUCT(BlueprintType)
+struct FDataLinkString
+{
+	GENERATED_BODY()
+
+	friend uint32 GetTypeHash(const FDataLinkString& InString)
+	{
+		return GetTypeHash(InString.Value);
+	}
+
+	bool operator==(const FDataLinkString& InString) const
+	{
+		return Value == InString.Value;
+	}
+
+	UPROPERTY(EditAnywhere, BlueprintReadWrite, Category="Data Link")
+	FString Value;
+};
+
+template<>
+struct TStructOpsTypeTraits<FDataLinkString> : public TStructOpsTypeTraitsBase2<FDataLinkString>
+{
+	enum
+	{
+		WithIdenticalViaEquality = true
+	};
+};
diff --git a/Engine/Plugins/Experimental/DataLink/Source/DataLink/Public/DataLinkDelegates.h b/Engine/Plugins/Experimental/DataLink/Source/DataLink/Public/DataLinkDelegates.h
new file mode 100644
--- /dev/null
+++ b/Engine/Plugins/Experimental/DataLink/Source/DataLink/Public/DataLinkDelegates.h
@@ -0,0 +1,12 @@
+// Copyright Epic Games, Inc. All Rights Reserved.
+
+#pragma once
+
+#include "Delegates/Delegate.h"
+#include "Delegates/DelegateCombinations.h"
+
+class FDataLinkExecutor;
+enum class EDataLinkExecutionResult : uint8;
+struct FConstStructView;
+
+DECLARE_DELEGATE_ThreeParams(FOnDataLinkExecutionFinished, const FDataLinkExecutor&, EDataLinkExecutionResult, FConstStructView);
diff --git a/Engine/Plugins/Experimental/DataLink/Source/DataLink/Public/DataLinkEnums.h b/Engine/Plugins/Experimental/DataLink/Source/DataLink/Public/DataLinkEnums.h
new file mode 100644
--- /dev/null
+++ b/Engine/Plugins/Experimental/DataLink/Source/DataLink/Public/DataLinkEnums.h
@@ -0,0 +1,20 @@
+// Copyright Epic Games, Inc. All Rights Reserved.
+
+#pragma once
+
+#include "HAL/Platform.h"
+#include "DataLinkEnums.generated.h"
+
+UENUM()
+enum class EDataLinkExecutionResult : uint8
+{
+	Failed,
+	Succeeded,
+};
+
+UENUM()
+enum class EDataLinkExecutionReply : uint8
+{
+	Unhandled,
+	Handled,
+};
diff --git a/Engine/Plugins/Experimental/DataLink/Source/DataLink/Public/DataLinkExecutor.h b/Engine/Plugins/Experimental/DataLink/Source/DataLink/Public/DataLinkExecutor.h
new file mode 100644
--- /dev/null
+++ b/Engine/Plugins/Experimental/DataLink/Source/DataLink/Public/DataLinkExecutor.h
@@ -0,0 +1,125 @@
+// Copyright Epic Games, Inc. All Rights Reserved.
+
+#pragma once
+
+#include "DataLinkDelegates.h"
+#include "DataLinkInstance.h"
+#include "DataLinkNodeInstance.h"
+#include "DataLinkSink.h"
+#include "StructUtils/InstancedStruct.h"
+
+class FDataLinkExecutorArguments;
+class UDataLinkGraph;
+class UDataLinkNode;
+struct FDataLinkPinReference;
+
+/**
+ * Executes a Data Link graph with custom instance data.
+ * @see FDataLinkExecutor::Run
+ */
+class FDataLinkExecutor : public FGCObject, public TSharedFromThis<FDataLinkExecutor>
+{
+	friend class FDataLinkExecutorArguments;
+
+	struct FPrivateToken { explicit FPrivateToken() = default; };
+
+public:
+	DATALINK_API static TSharedPtr<FDataLinkExecutor> Create(FDataLinkExecutorArguments&& InArgs);
+
+	explicit FDataLinkExecutor(FPrivateToken);
+
+	/** Gets the unique identifier for this Executor */
+	DATALINK_API const FGuid& GetExecutorId() const;
+
+	/**
+	 * Gets a brief description of the context for debugging purposes.
+	 * Returns a valid context name when WITH_DATALINK_CONTEXT is true.
+	 * Note: WITH_DATALINK_CONTEXT is true in builds with logging enabled.
+	 */
+	DATALINK_API FStringView GetContextName() const;
+
+	/** Gets the object that is responsible for this execution */
+	DATALINK_API UObject* GetContextObject() const;
+
+	/** Called to start the graph execution */
+	DATALINK_API void Run();
+
+	/** Retrieves the Node Instance Data of this execution for the given Node */
+	DATALINK_API const FDataLinkNodeInstance& GetNodeInstance(const UDataLinkNode* InNode) const;
+
+	DATALINK_API FDataLinkNodeInstance& GetNodeInstanceMutable(const UDataLinkNode* InNode);
+
+	/**
+	 * Called when a node has finished executing successfully
+	 * If it's the output node that has finished executing, it will finish the entire graph execution
+	 * @param InNode the node that finished execution
+	 * @param InOutputDataView the data view of the output (if succeeded)
+	 */
+	DATALINK_API void SucceedNode(const UDataLinkNode* InNode, FConstStructView InOutputDataView);
+
+	/**
+	 * Called when a node has failed to execute and produce a valid output
+	 * This stops the entire graph execution
+	 * @param InNode the node that failed
+	 */
+	DATALINK_API void FailNode(const UDataLinkNode* InNode);
+
+private:
+	/** Called when running to ensure all the parameters are properly set. Returns true if run can be done, false otherwise */
+	bool ValidateRun();
+
+	/** Called in Pin execution to make sure the pins and data views match in compatibility */
+	bool ValidateInputPins(TConstArrayView<FDataLinkPinReference> InInputPins, TConstArrayView<FConstStructView> InInputDataViews) const;
+
+	/** Called to execute the graph's input nodes */
+	bool ExecuteEntryNodes();
+
+	/**
+	 * Called to execute the given input pins with the input data
+	 * @param InInputPins the input pins to execute (can belong to different nodes)
+	 * @param InInputDataViews the matching data views for each input pin
+	 * @param OutNodesExecuted the number of nodes that got executed
+	 * @return true if the input pins were processed without error
+	 */
+	bool ExecuteInputPins(TConstArrayView<FDataLinkPinReference> InInputPins, TArray<FConstStructView>&& InInputDataViews, uint16* OutNodesExecuted = nullptr);
+
+	bool ExecuteNode(const UDataLinkNode& InNode);
+
+	FDataLinkNodeInstance& FindOrAddNodeInstance(const UDataLinkNode* InNode);
+
+	/** Called when a node has finished */
+	bool ProcessFinishedNode(const UDataLinkNode* InNode);
+
+	/** Called when the entire graph execution has completed */
+	void Finish(EDataLinkExecutionResult InStatus, FConstStructView InOutputData);
+
+	//~ Begin FGCObject
+	virtual FString GetReferencerName() const override;
+	virtual void AddReferencedObjects(FReferenceCollector& InCollector) override;
+	//~ End FGCObject
+
+	/** Unique Identifier for this Executor */
+	const FGuid ExecutorId;
+
+#if WITH_DATALINK_CONTEXT
+	/** Context string for additional information when logging */
+	FString ContextName;
+#endif
+
+	/** Object responsible for this execution */
+	TObjectPtr<UObject> ContextObject;
+
+	/** The data link instance of this execution */
+	FDataLinkInstance Instance;
+
+	/** Sink where all the data for this execution is queried and stored */
+	TSharedPtr<FDataLinkSink> Sink;
+
+	/** Delegate to call when the execution completes */
+	FOnDataLinkExecutionFinished OnFinishDelegate;
+
+	/** Map of a Node to its Instance (Input, Output and Instance Data) for this execution */
+	mutable TMap<TObjectPtr<const UDataLinkNode>, FDataLinkNodeInstance> NodeInstanceMap;
+
+	bool bRunning = false;
+};
diff --git a/Engine/Plugins/Experimental/DataLink/Source/DataLink/Public/DataLinkExecutorArguments.h b/Engine/Plugins/Experimental/DataLink/Source/DataLink/Public/DataLinkExecutorArguments.h
new file mode 100644
--- /dev/null
+++ b/Engine/Plugins/Experimental/DataLink/Source/DataLink/Public/DataLinkExecutorArguments.h
@@ -0,0 +1,48 @@
+// Copyright Epic Games, Inc. All Rights Reserved.
+
+#pragma once
+
+#include "Containers/UnrealString.h"
+#include "DataLinkDelegates.h"
+#include "DataLinkInstance.h"
+#include "StructUtils/StructView.h"
+
+class UDataLinkGraph;
+struct FDataLinkSink;
+
+class FDataLinkExecutorArguments
+{
+	friend class FDataLinkExecutor;
+
+public:
+	DATALINK_API FDataLinkExecutorArguments(const FDataLinkInstance& InInstance);
+
+	DATALINK_API FDataLinkExecutorArguments(FDataLinkInstance&& InInstance);
+
+#if WITH_DATALINK_CONTEXT
+	/** Added context information of the execution */
+	DATALINK_API FDataLinkExecutorArguments&& SetContextName(FString&& InContextName);
+#endif
+
+	/** Sets the object responsible for this execution */
+	DATALINK_API FDataLinkExecutorArguments&& SetContextObject(UObject* InContextObject);
+
+	/** Sets the Sink to use. If unset, the Executor will create its own temporary sink */
+	DATALINK_API FDataLinkExecutorArguments&& SetSink(const TSharedPtr<FDataLinkSink>& InSink);
+
+	/** The delegate to call when the execution finishes */
+	DATALINK_API FDataLinkExecutorArguments&& SetOnFinish(FOnDataLinkExecutionFinished&& InDelegate);
+
+private:
+	FDataLinkInstance Instance;
+
+#if WITH_DATALINK_CONTEXT
+	FString ContextName;
+#endif
+
+	TObjectPtr<UObject> ContextObject;
+
+	TSharedPtr<FDataLinkSink> Sink;
+
+	FOnDataLinkExecutionFinished OnFinishDelegate;
+};
diff --git a/Engine/Plugins/Experimental/DataLink/Source/DataLink/Public/DataLinkGraph.h b/Engine/Plugins/Experimental/DataLink/Source/DataLink/Public/DataLinkGraph.h
new file mode 100644
--- /dev/null
+++ b/Engine/Plugins/Experimental/DataLink/Source/DataLink/Public/DataLinkGraph.h
@@ -0,0 +1,79 @@
+// Copyright Epic Games, Inc. All Rights Reserved.
+
+#pragma once
+
+#include "Delegates/Delegate.h"
+#include "UObject/Object.h"
+#include "DataLinkGraph.generated.h"
+
+class UDataLinkNode;
+struct FConstStructView;
+struct FDataLinkPinReference;
+
+#if WITH_EDITOR
+class UEdGraph;
+#endif
+
+UCLASS(MinimalAPI, BlueprintType)
+class UDataLinkGraph : public UObject
+{
+	GENERATED_BODY()
+
+	friend class FDataLinkGraphCompiler;
+	friend class UDataLinkGraphFactory;
+
+public:
+	/** Counts the number of input pins within the input node list that are not linked to other nodes */
+	DATALINK_API int32 GetInputPinCount() const;
+
+	/** Retrieves the Input Pins from the input nodes that are not linked to other nodes */
+	DATALINK_API TArray<FDataLinkPinReference> GetInputPins() const;
+
+	/** Iterates the input pins from the input nodes that are not linked to other nodes. These serve as the 'input pin' for the graph */
+	DATALINK_API bool ForEachInputPin(TFunctionRef<bool(FDataLinkPinReference)> InFunction) const;
+
+	TConstArrayView<const UDataLinkNode*> GetInputNodes() const
+	{
+		return InputNodes;
+	}
+
+	const UDataLinkNode* GetOutputNode() const
+	{
+		return OutputNode;
+	}
+
+#if WITH_EDITOR
+	UEdGraph* GetEdGraph() const
+	{
+		return EdGraph;
+	}
+
+	static TMulticastDelegateRegistration<void(UDataLinkGraph*)>& OnGraphCompiled()
+	{
+		return OnGraphCompiledDelegate;
+	}
+#endif
+
+private:
+#if WITH_EDITOR
+	DATALINK_API static TMulticastDelegate<void(UDataLinkGraph*)> OnGraphCompiledDelegate;
+#endif
+
+	/** All the compiled nodes present in this graph */
+	UPROPERTY()
+	TArray<TObjectPtr<UDataLinkNode>> Nodes;
+
+	/** The nodes that the graph starts off with */
+	UPROPERTY()
+	TArray<TObjectPtr<UDataLinkNode>> InputNodes;
+
+	/** The node that provides the result data */
+	UPROPERTY()
+	TObjectPtr<UDataLinkNode> OutputNode;
+
+#if WITH_EDITORONLY_DATA
+	/** EdGraph used to compile the nodes in this graph */
+	UPROPERTY()
+	TObjectPtr<UEdGraph> EdGraph;
+#endif
+};
diff --git a/Engine/Plugins/Experimental/DataLink/Source/DataLink/Public/DataLinkInputDataViewer.h b/Engine/Plugins/Experimental/DataLink/Source/DataLink/Public/DataLinkInputDataViewer.h
new file mode 100644
--- /dev/null
+++ b/Engine/Plugins/Experimental/DataLink/Source/DataLink/Public/DataLinkInputDataViewer.h
@@ -0,0 +1,51 @@
+// Copyright Epic Games, Inc. All Rights Reserved.
+
+#pragma once
+
+#include "StructUtils/StructView.h"
+
+struct FDataLinkInputDataEntry
+{
+	bool operator==(FName InInputName) const
+	{
+		return Name == InInputName;
+	}
+
+	/** Name of the Input Data */
+	FName Name;
+
+	/** View to the Input Data */
+	FConstStructView DataView;
+};
+
+struct FDataLinkPin;
+
+class FDataLinkInputDataViewer
+{
+	friend class FDataLinkExecutor;
+
+public:
+	explicit FDataLinkInputDataViewer(TConstArrayView<FDataLinkPin> InInputPins);
+
+	DATALINK_API FConstStructView Find(FName InInputName) const;
+
+	int32 Num() const
+	{
+		return DataEntries.Num();
+	}
+
+	template<typename T>
+	const T& Get(FName InInputName) const
+	{
+		return Find(InInputName).Get<const T>();
+	}
+
+	TConstArrayView<FDataLinkInputDataEntry> GetDataEntries() const;
+
+private:
+	bool HasInvalidDataEntry() const;
+
+	void SetEntryData(const FDataLinkPin& InPin, FConstStructView InInputDataView);
+
+	TArray<FDataLinkInputDataEntry> DataEntries;
+};
diff --git a/Engine/Plugins/Experimental/DataLink/Source/DataLink/Public/DataLinkInstance.h b/Engine/Plugins/Experimental/DataLink/Source/DataLink/Public/DataLinkInstance.h
new file mode 100644
--- /dev/null
+++ b/Engine/Plugins/Experimental/DataLink/Source/DataLink/Public/DataLinkInstance.h
@@ -0,0 +1,24 @@
+// Copyright Epic Games, Inc. All Rights Reserved.
+
+#pragma once
+
+#include "Containers/Array.h"
+#include "StructUtils/InstancedStruct.h"
+#include "DataLinkInstance.generated.h"
+
+class UDataLinkGraph;
+
+/** Instance of a data link to be executed */
+USTRUCT(BlueprintType)
+struct FDataLinkInstance
+{
+	GENERATED_BODY()
+
+	/** The data link graph to execute */
+	UPROPERTY(EditAnywhere, BlueprintReadWrite, Category="Data Link")
+	TObjectPtr<UDataLinkGraph> DataLinkGraph;
+
+	/** The initial input data to feed into the data link graph */
+	UPROPERTY(EditAnywhere, EditFixedSize, BlueprintReadWrite, Category="Data Link")
+	TArray<FInstancedStruct> InputData;
+};
diff --git a/Engine/Plugins/Experimental/DataLink/Source/DataLink/Public/DataLinkLog.h b/Engine/Plugins/Experimental/DataLink/Source/DataLink/Public/DataLinkLog.h
new file mode 100644
--- /dev/null
+++ b/Engine/Plugins/Experimental/DataLink/Source/DataLink/Public/DataLinkLog.h
@@ -0,0 +1,7 @@
+// Copyright Epic Games, Inc. All Rights Reserved.
+
+#pragma once
+
+#include "Logging/LogMacros.h"
+
+DECLARE_LOG_CATEGORY_EXTERN(LogDataLink, Log, All);
diff --git a/Engine/Plugins/Experimental/DataLink/Source/DataLink/Public/DataLinkNames.h b/Engine/Plugins/Experimental/DataLink/Source/DataLink/Public/DataLinkNames.h
new file mode 100644
--- /dev/null
+++ b/Engine/Plugins/Experimental/DataLink/Source/DataLink/Public/DataLinkNames.h
@@ -0,0 +1,10 @@
+// Copyright Epic Games, Inc. All Rights Reserved.
+
+#pragma once
+
+#include "UObject/NameTypes.h"
+
+namespace UE::DataLink
+{
+	const FLazyName OutputDefault(TEXT("Output"));
+}
diff --git a/Engine/Plugins/Experimental/DataLink/Source/DataLink/Public/DataLinkNode.h b/Engine/Plugins/Experimental/DataLink/Source/DataLink/Public/DataLinkNode.h
new file mode 100644
--- /dev/null
+++ b/Engine/Plugins/Experimental/DataLink/Source/DataLink/Public/DataLinkNode.h
@@ -0,0 +1,84 @@
+// Copyright Epic Games, Inc. All Rights Reserved.
+
+#pragma once
+
+#include "DataLinkEnums.h"
+#include "DataLinkPin.h"
+#include "UObject/Object.h"
+#include "DataLinkNode.generated.h"
+
+class FDataLinkExecutor;
+struct FDataLinkPinBuilder;
+
+#if WITH_EDITOR
+struct FDataLinkNodeMetadata;
+#endif
+
+/**
+ * Base class for nodes in a data link graph, that handles the logic of the node.
+ * The instance data of the node can be gotten through the FDataLinkExecutor.
+ * Each node consists of an Input struct, and an Output struct.
+ * A data link node can be connected to another so long as the output struct of one is compatible with the input struct of the other.
+ */
+UCLASS(Abstract, MinimalAPI)
+class UDataLinkNode : public UObject
+{
+	friend class FDataLinkGraphCompiler;
+
+	GENERATED_BODY()
+
+public:
+	void Execute(FDataLinkExecutor& InExecutor) const;
+
+	TConstArrayView<FDataLinkPin> GetInputPins() const
+	{
+		return InputPins;
+	}
+
+	TConstArrayView<FDataLinkPin> GetOutputPins() const
+	{
+		return OutputPins;
+	}
+
+	const UScriptStruct* GetInstanceStruct() const
+	{
+		return InstanceStruct;
+	}
+
+#if WITH_EDITOR
+	DATALINK_API void BuildMetadata(FDataLinkNodeMetadata& OutMetadata) const;
+#endif
+
+	DATALINK_API void BuildPins(TArray<FDataLinkPin>& OutInputPins, TArray<FDataLinkPin>& OutOutputPins) const;
+
+protected:
+#if WITH_EDITOR
+	/** Call to retrieve the metadata for this node */
+	virtual void OnBuildMetadata(FDataLinkNodeMetadata& Metadata) const {}
+#endif
+
+	/** Call to setup the pins required by this node */
+	virtual void OnBuildPins(FDataLinkPinBuilder& Inputs, FDataLinkPinBuilder& Outputs) const {}
+
+	/**
+	 * Execution logic of the node
+	 * @return whether the implementation handled the execution (regardless if it has finished or not)
+	 */
+	virtual EDataLinkExecutionReply OnExecute(FDataLinkExecutor& InExecutor) const
+	{
+		return EDataLinkExecutionReply::Unhandled;
+	}
+
+	/** Optional instance data struct for data outside input and output */
+	UPROPERTY()
+	TObjectPtr<const UScriptStruct> InstanceStruct;
+
+private:
+	/** Input Pins to start the execution on this Data Link node */
+	UPROPERTY()
+	TArray<FDataLinkPin> InputPins;
+
+	/** Output pins that feed the result of this node's execution */
+	UPROPERTY()
+	TArray<FDataLinkPin> OutputPins;
+};
diff --git a/Engine/Plugins/Experimental/DataLink/Source/DataLink/Public/DataLinkNodeInstance.h b/Engine/Plugins/Experimental/DataLink/Source/DataLink/Public/DataLinkNodeInstance.h
new file mode 100644
--- /dev/null
+++ b/Engine/Plugins/Experimental/DataLink/Source/DataLink/Public/DataLinkNodeInstance.h
@@ -0,0 +1,74 @@
+// Copyright Epic Games, Inc. All Rights Reserved.
+
+#pragma once
+
+#include "DataLinkInputDataViewer.h"
+#include "DataLinkOutputDataViewer.h"
+#include "DataLinkSink.h"
+#include "StructUtils/InstancedStruct.h"
+#include "StructUtils/StructView.h"
+
+class FReferenceCollector;
+
+enum class EDataLinkNodeStatus : uint8
+{
+	NotStarted,
+	Executing,
+	Finished,
+};
+
+struct FDataLinkNodeInstance
+{
+	friend class FDataLinkExecutor;
+
+	explicit FDataLinkNodeInstance(const UDataLinkNode& InNode);
+
+	void AddReferencedObjects(FReferenceCollector& InCollector);
+
+	const FDataLinkInputDataViewer& GetInputDataViewer() const
+	{
+		return InputDataViewer;
+	}
+
+	const FDataLinkOutputDataViewer& GetOutputDataViewer() const
+	{
+		return OutputDataViewer;
+	}
+
+	FConstStructView GetInstanceData() const
+	{
+		return InstanceData;
+	}
+
+	FStructView GetInstanceDataMutable()
+	{
+		return InstanceData;
+	}
+
+	const FDataLinkSinkKey& GetSinkKey() const
+	{
+		return SinkKey;
+	}
+
+private:
+	/**
+	 * The Sink Key for this Node Instance
+	 * Saved here to avoid recreating if it needs to be re-used.
+	 */
+	FDataLinkSinkKey SinkKey;
+
+	/** Views of the input data, matching the Node's Input Pins */
+	FDataLinkInputDataViewer InputDataViewer;
+
+	/** Views of the output data, matching the Node's Output Pins */
+	FDataLinkOutputDataViewer OutputDataViewer;
+
+	/**
+	 * Optional data within the node instanced for every execution.
+	 * This is used to store data outside of the input/output data.
+	 */
+	FInstancedStruct InstanceData;
+
+	/** Current Status of a Node */
+	EDataLinkNodeStatus Status = EDataLinkNodeStatus::NotStarted;
+};
diff --git a/Engine/Plugins/Experimental/DataLink/Source/DataLink/Public/DataLinkNodeMetadata.h b/Engine/Plugins/Experimental/DataLink/Source/DataLink/Public/DataLinkNodeMetadata.h
new file mode 100644
--- /dev/null
+++ b/Engine/Plugins/Experimental/DataLink/Source/DataLink/Public/DataLinkNodeMetadata.h
@@ -0,0 +1,42 @@
+// Copyright Epic Games, Inc. All Rights Reserved.
+
+#pragma once
+
+#if WITH_EDITOR
+#include "Internationalization/Text.h"
+#endif
+
+#if WITH_EDITOR
+/** Metadata of a Node for Editor-side representation. */
+struct FDataLinkNodeMetadata
+{
+	FDataLinkNodeMetadata& SetDisplayName(const FText& InDisplayName)
+	{
+		DisplayName = InDisplayName;
+		return *this;
+	}
+
+	FDataLinkNodeMetadata& SetTooltipText(const FText& InDescription)
+	{
+		TooltipText = InDescription;
+		return *this;
+	}
+
+	const FText& GetDisplayName() const
+	{
+		return DisplayName;
+	}
+
+	const FText& GetTooltipText() const
+	{
+		return TooltipText;
+	}
+
+private:
+	/** Display name of the Node */
+	FText DisplayName;
+
+	/** Description for the node. Used for the Tooltip text of nodes */
+	FText TooltipText;
+};
+#endif
diff --git a/Engine/Plugins/Experimental/DataLink/Source/DataLink/Public/DataLinkOutputDataViewer.h b/Engine/Plugins/Experimental/DataLink/Source/DataLink/Public/DataLinkOutputDataViewer.h
new file mode 100644
--- /dev/null
+++ b/Engine/Plugins/Experimental/DataLink/Source/DataLink/Public/DataLinkOutputDataViewer.h
@@ -0,0 +1,46 @@
+// Copyright Epic Games, Inc. All Rights Reserved.
+
+#pragma once
+
+#include "StructUtils/StructView.h"
+
+struct FDataLinkPin;
+
+struct FDataLinkOutputDataEntry
+{
+	bool operator==(FName InOutputName) const
+	{
+		return Name == InOutputName;
+	}
+
+	/** Name of the Output Data */
+	FName Name;
+
+	FStructView GetDataView(const UScriptStruct* InDesiredStruct) const;
+
+private:
+	/** Output Data instantiated */
+	mutable FInstancedStruct OutputData;
+};
+
+class FDataLinkOutputDataViewer
+{
+public:
+	explicit FDataLinkOutputDataViewer(TConstArrayView<FDataLinkPin> InOutputPins);
+
+	DATALINK_API FStructView Find(FName InOutputName, const UScriptStruct* InDesiredStruct) const;
+
+	int32 Num() const
+	{
+		return DataEntries.Num();
+	}
+
+	template<typename T>
+	T& Get(FName InOutputName) const
+	{
+		return this->Find(InOutputName, T::StaticStruct()).template Get<T>();
+	}
+
+private:
+	TArray<FDataLinkOutputDataEntry> DataEntries;
+};
diff --git a/Engine/Plugins/Experimental/DataLink/Source/DataLink/Public/DataLinkPin.h b/Engine/Plugins/Experimental/DataLink/Source/DataLink/Public/DataLinkPin.h
new file mode 100644
--- /dev/null
+++ b/Engine/Plugins/Experimental/DataLink/Source/DataLink/Public/DataLinkPin.h
@@ -0,0 +1,48 @@
+// Copyright Epic Games, Inc. All Rights Reserved.
+
+#pragma once
+
+#include "UObject/NameTypes.h"
+#include "UObject/ObjectPtr.h"
+#include "DataLinkPin.generated.h"
+
+class UDataLinkNode;
+
+USTRUCT()
+struct FDataLinkPin
+{
+	GENERATED_BODY()
+
+	FDataLinkPin() = default;
+
+	DATALINK_API explicit FDataLinkPin(FName InPinName);
+
+	DATALINK_API const FDataLinkPin* GetLinkedInputPin() const;
+
+	DATALINK_API FText GetDisplayName() const;
+
+	bool operator==(const FName InName) const
+	{
+		return Name == InName;
+	}
+
+	/** Unique name for the Pin */
+	UPROPERTY()
+	FName Name;
+
+	/** Display name of the Pin */
+	UPROPERTY()
+	FText DisplayName;
+
+	/** Struct type of the Pin */
+	UPROPERTY()
+	TObjectPtr<const UScriptStruct> Struct;
+
+	/** The node this pin connects to */
+	UPROPERTY()
+	TObjectPtr<const UDataLinkNode> LinkedNode;
+
+	/** the pin index at which this pin connects to the linked node (to find linked pin) */
+	UPROPERTY()
+	int32 LinkedIndex = INDEX_NONE;
+};
diff --git a/Engine/Plugins/Experimental/DataLink/Source/DataLink/Public/DataLinkPinBuilder.h b/Engine/Plugins/Experimental/DataLink/Source/DataLink/Public/DataLinkPinBuilder.h
new file mode 100644
--- /dev/null
+++ b/Engine/Plugins/Experimental/DataLink/Source/DataLink/Public/DataLinkPinBuilder.h
@@ -0,0 +1,55 @@
+// Copyright Epic Games, Inc. All Rights Reserved.
+
+#pragma once
+
+#include "Containers/Array.h"
+#include "DataLinkPin.h"
+
+/** Short-lived parameter building for a single pin */
+struct FDataLinkPinParameters
+{
+	DATALINK_API explicit FDataLinkPinParameters(FDataLinkPin& InPin);
+
+	/** 
+	 * Optionally sets the display text of the pin
+	 * If not set, it will use the pin's unique name instead.
+	 */
+	DATALINK_API FDataLinkPinParameters& SetDisplayName(FText&& InPinDisplayName);
+
+	template<typename InStruct UE_REQUIRES(TModels_V<CStaticStructProvider, InStruct>)>
+	FDataLinkPinParameters& SetStruct()
+	{
+		return this->SetStruct(InStruct::StaticStruct());
+	}
+
+	/**
+	 * Sets the struct of the pin.
+	 * Can be left null for cases like output pins where the struct might not be known
+	 */
+	DATALINK_API FDataLinkPinParameters& SetStruct(const UScriptStruct* InPinStruct);
+
+private:
+	FDataLinkPin& Pin;
+};
+
+/** Short-lived builder for pin arrays limiting to only add pins with unique names */
+struct FDataLinkPinBuilder
+{
+	explicit FDataLinkPinBuilder(TArray<FDataLinkPin>& OutPins);
+
+	/**
+	 * Possibly re-allocates the array to ensure its capacity can hold the number of items to add without further re-allocations
+	 * @param InNumToAdd the number of new elements to add
+	 */
+	DATALINK_API void AddCapacity(int32 InNumToAdd);
+
+	/**
+	 * Adds a new pin to the array with a given unique name
+	 * @param InPinName the unique name of the pin
+	 * @return parameter object to further customize the pin
+	 */
+	DATALINK_API FDataLinkPinParameters Add(FName InPinName);
+
+private:
+	TArray<FDataLinkPin>& Pins;
+};
diff --git a/Engine/Plugins/Experimental/DataLink/Source/DataLink/Public/DataLinkPinReference.h b/Engine/Plugins/Experimental/DataLink/Source/DataLink/Public/DataLinkPinReference.h
new file mode 100644
--- /dev/null
+++ b/Engine/Plugins/Experimental/DataLink/Source/DataLink/Public/DataLinkPinReference.h
@@ -0,0 +1,22 @@
+// Copyright Epic Games, Inc. All Rights Reserved.
+
+#pragma once
+
+#include "UObject/ObjectPtr.h"
+
+class UDataLinkNode;
+struct FDataLinkPin;
+
+/** Short-lived data containing a pin reference with the Node that owns it. */
+struct FDataLinkPinReference
+{
+	explicit FDataLinkPinReference(const UDataLinkNode* InOwningNode, const FDataLinkPin* InPin)
+		: OwningNode(InOwningNode)
+		, Pin(InPin)
+	{
+	}
+
+	TObjectPtr<const UDataLinkNode> OwningNode;
+
+	const FDataLinkPin* Pin = nullptr;
+};
diff --git a/Engine/Plugins/Experimental/DataLink/Source/DataLink/Public/DataLinkRequestProxy.h b/Engine/Plugins/Experimental/DataLink/Source/DataLink/Public/DataLinkRequestProxy.h
new file mode 100644
--- /dev/null
+++ b/Engine/Plugins/Experimental/DataLink/Source/DataLink/Public/DataLinkRequestProxy.h
@@ -0,0 +1,38 @@
+// Copyright Epic Games, Inc. All Rights Reserved.
+
+#pragma once
+
+#include "StructUtils/InstancedStruct.h"
+#include "UObject/Object.h"
+#include "DataLinkRequestProxy.generated.h"
+
+class IDataLinkSinkProvider;
+class FDataLinkExecutor;
+enum class EDataLinkExecutionResult : uint8;
+struct FDataLinkInstance;
+
+DECLARE_DYNAMIC_MULTICAST_DELEGATE_TwoParams(FOnDataLinkRequestComplete, const FInstancedStruct&, OutputData, EDataLinkExecutionResult, ExecutionResult);
+
+UCLASS(MinimalAPI)
+class UDataLinkRequestProxy : public UObject
+{
+	GENERATED_BODY()
+
+public:
+	UPROPERTY(BlueprintAssignable)
+	FOnDataLinkRequestComplete OnRequestComplete;
+
+	UFUNCTION(BlueprintCallable, Category="Data Link", meta=(BlueprintInternalUseOnly))
+	static DATALINK_API UDataLinkRequestProxy* CreateRequestProxy(FDataLinkInstance InDataLinkInstance
+		, UObject* InExecutionContext
+		, TScriptInterface<IDataLinkSinkProvider> InDataLinkSinkProvider);
+
+	DATALINK_API void ProcessRequest(FDataLinkInstance&& InDataLinkInstance
+		, UObject* InExecutionContext
+		, TScriptInterface<IDataLinkSinkProvider> InDataLinkSinkProvider);
+
+private:
+	void OnExecutionFinished(const FDataLinkExecutor& InExecutor, EDataLinkExecutionResult InResult, FConstStructView InOutputData);
+
+	TSharedPtr<FDataLinkExecutor> DataLinkExecutor;
+};
diff --git a/Engine/Plugins/Experimental/DataLink/Source/DataLink/Public/DataLinkSink.h b/Engine/Plugins/Experimental/DataLink/Source/DataLink/Public/DataLinkSink.h
new file mode 100644
--- /dev/null
+++ b/Engine/Plugins/Experimental/DataLink/Source/DataLink/Public/DataLinkSink.h
@@ -0,0 +1,61 @@
+// Copyright Epic Games, Inc. All Rights Reserved.
+
+#pragma once
+
+#include "DataLinkInputDataViewer.h"
+#include "Misc/Optional.h"
+#include "StructUtils/StructView.h"
+#include "UObject/ObjectKey.h"
+
+class FDataLinkInputDataViewer;
+class UDataLinkNode;
+struct FDataLinkNodeInstance;
+
+/** Struct used as the Key to access the data for a given node and input data */
+struct FDataLinkSinkKey
+{
+	FDataLinkSinkKey() = default;
+
+	explicit FDataLinkSinkKey(TSubclassOf<UDataLinkNode> InNodeClass, const FDataLinkInputDataViewer& InInputDataViewer);
+
+	bool operator==(const FDataLinkSinkKey& InOther) const;
+
+	friend uint64 GetTypeHash(const FDataLinkSinkKey& InKey);
+
+private:
+	/** Class of the node that this data will be used in */
+	FObjectKey NodeClass;
+
+	/** Views of the input data. */
+	TArray<FDataLinkInputDataEntry, TInlineAllocator<2>> InputDataEntries;
+
+	/** Hash computed in ctor */
+	uint64 CachedHash = 0;
+};
+
+/** Sink that caches node data to be reusable by nodes with the same logic (i.e. "class") and input data */
+struct FDataLinkSink
+{
+	/**
+	 * Returns the output data view that matches the node class and the settings within the input data view.
+	 * @return a data view to the output data
+	 */
+	DATALINK_API FInstancedStruct& FindOrAddCachedData(const FDataLinkNodeInstance& InNodeInstance, TOptional<const UScriptStruct*> InDesiredStruct = {});
+
+	template<typename T>
+	T& FindOrAddCachedData(const FDataLinkNodeInstance& InNodeInstance)
+	{
+		FInstancedStruct& InstancedStruct = this->FindOrAddCachedData(InNodeInstance, T::StaticStruct());
+		return InstancedStruct.GetMutable<T>();
+	}
+
+	void AddStructReferencedObjects(FReferenceCollector& InCollector);
+
+private:
+	/**
+	 * Map of the Node & Input Data to the Cached Data
+	 * This stores output data for Input data regardless if it's hashable or not.
+	 * For input data without a type hash implementation, it should return a type hash of 0, generating collisions that then should get the correct element via comparison.
+	 */
+	TMap<FDataLinkSinkKey, FInstancedStruct> CachedDataMap;
+};
diff --git a/Engine/Plugins/Experimental/DataLink/Source/DataLink/Public/DataLinkSinkObject.h b/Engine/Plugins/Experimental/DataLink/Source/DataLink/Public/DataLinkSinkObject.h
new file mode 100644
--- /dev/null
+++ b/Engine/Plugins/Experimental/DataLink/Source/DataLink/Public/DataLinkSinkObject.h
@@ -0,0 +1,29 @@
+// Copyright Epic Games, Inc. All Rights Reserved.
+
+#pragma once
+
+#include "IDataLinkSinkProvider.h"
+#include "UObject/Object.h"
+#include "DataLinkSinkObject.generated.h"
+
+struct FDataLinkSink;
+
+UCLASS(MinimalAPI, BlueprintType)
+class UDataLinkSinkObject : public UObject, public IDataLinkSinkProvider
+{
+	GENERATED_BODY()
+
+public:
+	UDataLinkSinkObject();
+
+	UFUNCTION(BlueprintCallable, Category="Data Link")
+	void ResetSink();
+
+	//~ Begin IDataLinkSinkProvider
+	virtual TSharedPtr<FDataLinkSink> GetSink() const override;
+	virtual const UDataLinkSinkObject* GetSinkObject_Implementation() const override;
+	//~ End IDataLinkSinkProvider
+
+private:
+	TSharedPtr<FDataLinkSink> Sink;
+};
diff --git a/Engine/Plugins/Experimental/DataLink/Source/DataLink/Public/DataLinkUtils.h b/Engine/Plugins/Experimental/DataLink/Source/DataLink/Public/DataLinkUtils.h
new file mode 100644
--- /dev/null
+++ b/Engine/Plugins/Experimental/DataLink/Source/DataLink/Public/DataLinkUtils.h
@@ -0,0 +1,31 @@
+// Copyright Epic Games, Inc. All Rights Reserved.
+
+#pragma once
+
+#include "HAL/Platform.h"
+#include "Templates/SharedPointerFwd.h"
+
+class IDataLinkSinkProvider;
+class UClass;
+class UObject;
+struct FDataLinkSink;
+template <typename T> class TScriptInterface;
+
+namespace UE::DataLink
+{
+	/**
+	 * Replaces a given Object with a new object with the same name but different class
+	 * @param InOutObject object to replace. Can come in as null
+	 * @param InOuter outer of the object. Must be valid.
+	 * @param InClass new class of the new object that will replace the older one
+	 * @return true if the operation took place, false otherwise
+	 */
+	DATALINK_API bool ReplaceObject(UObject*& InOutObject, UObject* InOuter, UClass* InClass);
+
+	/**
+	 * Attempts to get the underlying Sink from the given Sink Provider
+	 * @param InSinkProvider the sink provider interface
+	 * @return the sink if found
+	 */
+	DATALINK_API TSharedPtr<FDataLinkSink> TryGetSink(TScriptInterface<IDataLinkSinkProvider> InSinkProvider);
+}
diff --git a/Engine/Plugins/Experimental/DataLink/Source/DataLink/Public/IDataLinkSinkProvider.h b/Engine/Plugins/Experimental/DataLink/Source/DataLink/Public/IDataLinkSinkProvider.h
new file mode 100644
--- /dev/null
+++ b/Engine/Plugins/Experimental/DataLink/Source/DataLink/Public/IDataLinkSinkProvider.h
@@ -0,0 +1,29 @@
+// Copyright Epic Games, Inc. All Rights Reserved.
+
+#pragma once
+
+#include "UObject/Interface.h"
+#include "IDataLinkSinkProvider.generated.h"
+
+class UDataLinkSinkObject;
+struct FDataLinkSink;
+
+UINTERFACE(MinimalAPI)
+class UDataLinkSinkProvider : public UInterface
+{
+	GENERATED_BODY()
+};
+
+class IDataLinkSinkProvider
+{
+	GENERATED_BODY()
+
+public:
+	UFUNCTION(BlueprintNativeEvent, Category="Data Link")
+	DATALINK_API const UDataLinkSinkObject* GetSinkObject() const;
+
+	virtual TSharedPtr<FDataLinkSink> GetSink() const
+	{
+		return nullptr;
+	}
+};
diff --git a/Engine/Plugins/Experimental/DataLink/Source/DataLink/Public/Nodes/DataLinkConstant.h b/Engine/Plugins/Experimental/DataLink/Source/DataLink/Public/Nodes/DataLinkConstant.h
new file mode 100644
--- /dev/null
+++ b/Engine/Plugins/Experimental/DataLink/Source/DataLink/Public/Nodes/DataLinkConstant.h
@@ -0,0 +1,33 @@
+// Copyright Epic Games, Inc. All Rights Reserved.
+
+#pragma once
+
+#include "DataLinkNode.h"
+#include "StructUtils/InstancedStruct.h"
+#include "DataLinkConstant.generated.h"
+
+/** Constants are a no-input node that provide a struct instance that does not change in execution time */
+UCLASS(MinimalAPI, DisplayName="Constant", Category="Core")
+class UDataLinkConstant : public UDataLinkNode
+{
+	GENERATED_BODY()
+
+public:
+	DATALINK_API void SetStruct(const UScriptStruct* InStructType);
+
+protected:
+	//~ Begin UDataLinkNode
+#if WITH_EDITOR
+	DATALINK_API virtual void OnBuildMetadata(FDataLinkNodeMetadata& Metadata) const override;
+#endif
+	DATALINK_API virtual void OnBuildPins(FDataLinkPinBuilder& Inputs, FDataLinkPinBuilder& Outputs) const override;
+	DATALINK_API virtual EDataLinkExecutionReply OnExecute(FDataLinkExecutor& InExecutor) const override;
+	//~ End UDataLinkNode
+
+private:
+	UPROPERTY(EditAnywhere, Category="Data Link", meta=(InvalidatesNode))
+	FText DisplayName; 
+
+	UPROPERTY(EditAnywhere, Category="Data Link", meta=(InvalidatesNode))
+	FInstancedStruct Instance;
+};
diff --git a/Engine/Plugins/Experimental/DataLink/Source/DataLink/Public/Nodes/Script/DataLinkScriptNode.h b/Engine/Plugins/Experimental/DataLink/Source/DataLink/Public/Nodes/Script/DataLinkScriptNode.h
new file mode 100644
--- /dev/null
+++ b/Engine/Plugins/Experimental/DataLink/Source/DataLink/Public/Nodes/Script/DataLinkScriptNode.h
@@ -0,0 +1,82 @@
+// Copyright Epic Games, Inc. All Rights Reserved.
+
+#pragma once
+
+#include "DataLinkNode.h"
+#include "DataLinkScriptNode.generated.h"
+
+USTRUCT()
+struct FDataLinkScriptPin
+{
+	GENERATED_BODY()
+
+	/** Unique name for the Pin */
+	UPROPERTY(EditAnywhere, Category = "DataLink")
+	FName Name;
+
+	/** Struct type of the Pin */
+	UPROPERTY(EditAnywhere, Category = "DataLink")
+	TObjectPtr<const UScriptStruct> Struct;
+};
+
+/**
+ * Script Nodes are a Blueprint Implementation of a Data Link Node.
+ * It does not inherit from UDataLinkNode, as the wrapper does this and forwards the logic execution here.
+ * This is done to allow for blueprints mutable nature (e.g. allowing users to set variables)
+ * @see UDataLinkScriptNodeWrapper
+ */
+UCLASS(Abstract, MinimalAPI, Blueprintable)
+class UDataLinkScriptNode : public UObject
+{
+	GENERATED_BODY()
+
+	friend class UDataLinkScriptNodeWrapper;
+
+public:
+	void Execute(const UDataLinkNode* InNode, FDataLinkExecutor& InExecutor);
+
+	UFUNCTION(BlueprintImplementableEvent)
+	void OnExecute();
+
+	/** Called to move to the next node to execute (or finish if last node) providing the output data as an instanced struct */
+	UFUNCTION(BlueprintCallable, Category="Data Link", DisplayName="Succeed")
+	bool Succeed(const FInstancedStruct& OutputData);
+
+	/** Called to move to the next node to execute (or finish if last node) and providing the output data as a wildcard */
+	UFUNCTION(BlueprintCallable, CustomThunk, Category="Data Link", DisplayName="Succeed (Wildcard)", meta=(CustomStructureParam="OutputData"))
+	bool SucceedWildcard(int32 OutputData);
+
+	/** Called to fail execution */
+	UFUNCTION(BlueprintCallable, Category="Data Link")
+	void Fail();
+
+	/** Retrieves the input data as an Instanced Struct */
+	UFUNCTION(BlueprintCallable, Category="Data Link", DisplayName="Get Input Data")
+	bool GetInputData(FInstancedStruct& InputData, FName InputName) const;
+
+	/** Retrieves the input data as a Wildcard */
+	UFUNCTION(BlueprintCallable, CustomThunk, Category="Data Link", DisplayName="Get Input Data (Wildcard)", meta=(CustomStructureParam="InputData"))
+	bool GetInputDataWildcard(int32& InputData, FName InputName) const;
+
+protected:
+	//~ Begin UObject
+	virtual UWorld* GetWorld() const override;
+	//~ End UObject
+
+private:
+	UObject* GetContextObject() const;
+
+	DECLARE_FUNCTION(execSucceedWildcard);
+	DECLARE_FUNCTION(execGetInputDataWildcard);
+
+	UPROPERTY(EditDefaultsOnly, Category="Data Link")
+	TArray<FDataLinkScriptPin> InputPins;
+
+	UPROPERTY(EditDefaultsOnly, Category="Data Link")
+	FDataLinkScriptPin OutputPin;
+
+	UPROPERTY()
+	TObjectPtr<const UDataLinkNode> Node;
+
+	TWeakPtr<FDataLinkExecutor> ExecutorWeak;
+};
diff --git a/Engine/Plugins/Experimental/DataLink/Source/DataLink/Public/Nodes/Script/DataLinkScriptNodeWrapper.h b/Engine/Plugins/Experimental/DataLink/Source/DataLink/Public/Nodes/Script/DataLinkScriptNodeWrapper.h
new file mode 100644
--- /dev/null
+++ b/Engine/Plugins/Experimental/DataLink/Source/DataLink/Public/Nodes/Script/DataLinkScriptNodeWrapper.h
@@ -0,0 +1,64 @@
+// Copyright Epic Games, Inc. All Rights Reserved.
+
+#pragma once
+
+#include "DataLinkNode.h"
+#include "Templates/SubclassOf.h"
+#include "DataLinkScriptNodeWrapper.generated.h"
+
+class UDataLinkScriptNode;
+
+/** Wrapper struct for the Node Instance containing the actual blueprint implementation of the data link logic */
+USTRUCT()
+struct FDataLinkScriptNodeInstance
+{
+	GENERATED_BODY()
+
+	UPROPERTY()
+	TObjectPtr<UDataLinkScriptNode> Node;
+};
+
+/**
+ * Wrapper to the actual blueprint implementation of the node.
+ * Data Link Nodes are by-design immutable in execution, and only modify input/output/instance data provided by the executor.
+ * Data Link Script Nodes are mutable (like typical blueprints) so their actual implementation (UDataLinkScriptNode) is instanced in every execution
+ * @see UDataLinkScriptNode
+ */
+UCLASS(MinimalAPI, DisplayName="Blueprint Node", meta=(Hidden))
+class UDataLinkScriptNodeWrapper : public UDataLinkNode
+{
+	GENERATED_BODY()
+
+public:
+	DATALINK_API UDataLinkScriptNodeWrapper();
+
+	static FName GetNodeClassPropertyName()
+	{
+		return GET_MEMBER_NAME_CHECKED(UDataLinkScriptNodeWrapper, NodeClass);
+	}
+
+	DATALINK_API void SetNodeClass(TSubclassOf<UDataLinkScriptNode> InNodeClass);
+
+protected:
+	//~ Begin UDataLinkNode
+#if WITH_EDITOR
+	DATALINK_API virtual void OnBuildMetadata(FDataLinkNodeMetadata& Metadata) const override;
+#endif
+	DATALINK_API virtual void OnBuildPins(FDataLinkPinBuilder& Inputs, FDataLinkPinBuilder& Outputs) const override;
+	DATALINK_API virtual EDataLinkExecutionReply OnExecute(FDataLinkExecutor& InExecutor) const override;
+	//~ End UDataLinkNode
+
+private:
+	UDataLinkScriptNode* GetDefaultNode() const;
+
+	/**
+	 * The node class to instantiate.
+	 * The selected node class determines the pins of this node.
+	 */
+	UPROPERTY(meta=(InvalidatesNode))
+	TSubclassOf<UDataLinkScriptNode> NodeClass;
+
+	/** Instance of Node Class serving as template for the execution instances */
+	UPROPERTY(VisibleAnywhere, Instanced, Category="Data Link")
+	TObjectPtr<UDataLinkScriptNode> TemplateNode;
+};
diff --git a/Engine/Plugins/Experimental/DataLink/Source/DataLink/Public/Nodes/String/DataLinkNodeStringBuilder.h b/Engine/Plugins/Experimental/DataLink/Source/DataLink/Public/Nodes/String/DataLinkNodeStringBuilder.h
new file mode 100644
--- /dev/null
+++ b/Engine/Plugins/Experimental/DataLink/Source/DataLink/Public/Nodes/String/DataLinkNodeStringBuilder.h
@@ -0,0 +1,35 @@
+// Copyright Epic Games, Inc. All Rights Reserved.
+
+#pragma once
+
+#include "DataLinkNode.h"
+#include "DataLinkStringBuilder.h"
+#include "DataLinkNodeStringBuilder.generated.h"
+
+/** Builder Node to help form a more complex string with parameters */
+UCLASS(MinimalAPI, DisplayName="String Builder", Category="Core")
+class UDataLinkNodeStringBuilder : public UDataLinkNode
+{
+	GENERATED_BODY()
+
+protected:
+	//~ Begin UDataLinkNode
+	DATALINK_API virtual void OnBuildPins(FDataLinkPinBuilder& Inputs, FDataLinkPinBuilder& Outputs) const override;
+	DATALINK_API virtual EDataLinkExecutionReply OnExecute(FDataLinkExecutor& InExecutor) const override;
+	//~ End UDataLinkNode
+
+	//~ Begin UObject
+#if WITH_EDITOR
+	DATALINK_API virtual void PostEditChangeProperty(FPropertyChangedEvent& InPropertyChangedEvent) override;
+#endif
+	//~ End UObject
+
+private:
+	/** The Segments to build, where a token can be defined as {Token Name} in its own isolated entry */
+	UPROPERTY(EditAnywhere, Category="Data Link", meta=(InvalidatesNode))
+	TArray<FString> Segments;
+
+	/** The tokens found within the Segment Array */
+	UPROPERTY()
+	TArray<FDataLinkStringBuilderToken> Tokens;
+};
diff --git a/Engine/Plugins/Experimental/DataLink/Source/DataLink/Public/Nodes/String/DataLinkStringBuilder.h b/Engine/Plugins/Experimental/DataLink/Source/DataLink/Public/Nodes/String/DataLinkStringBuilder.h
new file mode 100644
--- /dev/null
+++ b/Engine/Plugins/Experimental/DataLink/Source/DataLink/Public/Nodes/String/DataLinkStringBuilder.h
@@ -0,0 +1,53 @@
+// Copyright Epic Games, Inc. All Rights Reserved.
+
+#pragma once
+
+#include "Containers/ArrayView.h"
+#include "Containers/StringFwd.h"
+#include "Templates/FunctionFwd.h"
+#include "DataLinkStringBuilder.generated.h"
+
+class FDataLinkInputDataViewer;
+class FString;
+struct FConstStructView;
+struct FDataLinkPinBuilder;
+
+USTRUCT()
+struct FDataLinkStringBuilderToken
+{
+	GENERATED_BODY()
+
+	/** Name of the token. Used to build the Input Pin */
+	UPROPERTY()
+	FName Name;
+
+	/** The indices this token is located in the Segment array */
+	UPROPERTY()
+	TArray<int32> Indices;
+};
+
+/**
+ * Logic for String building for re-usability between different nodes without needing to inherit from each other
+ * @see UDataLinkNodeStringBuilder
+ */
+struct FDataLinkStringBuilder
+{
+	explicit FDataLinkStringBuilder(TConstArrayView<FString> InSegments, TConstArrayView<FDataLinkStringBuilderToken> InTokens)
+		: Segments(InSegments)
+		, Tokens(InTokens)
+	{
+	}
+
+	DATALINK_API bool BuildString(const FDataLinkInputDataViewer& InTokenValues, FString& InOutResult) const;
+
+	DATALINK_API void BuildInputPins(FDataLinkPinBuilder& Inputs) const;
+
+	DATALINK_API static void GatherTokens(TConstArrayView<FString> InSegments, TArray<FDataLinkStringBuilderToken>& OutTokens);
+
+private:
+	void ForEachResolvedSegment(const FDataLinkInputDataViewer& InTokenValues, TFunctionRef<void(FStringView InResolvedSegment)> InFunction) const;
+
+	TConstArrayView<FString> Segments;
+
+	TConstArrayView<FDataLinkStringBuilderToken> Tokens;
+};
diff --git a/Engine/Plugins/Experimental/DataLink/Source/DataLinkDataTable/DataLinkDataTable.Build.cs b/Engine/Plugins/Experimental/DataLink/Source/DataLinkDataTable/DataLinkDataTable.Build.cs
new file mode 100644
--- /dev/null
+++ b/Engine/Plugins/Experimental/DataLink/Source/DataLinkDataTable/DataLinkDataTable.Build.cs
@@ -0,0 +1,21 @@
+// Copyright Epic Games, Inc. All Rights Reserved.
+
+using UnrealBuildTool;
+
+public class DataLinkDataTable : ModuleRules
+{
+    public DataLinkDataTable(ReadOnlyTargetRules Target) : base(Target)
+    {
+        PCHUsage = ModuleRules.PCHUsageMode.UseExplicitOrSharedPCHs;
+
+        PublicDependencyModuleNames.AddRange(
+            new string[]
+            {
+                "Core", 
+                "CoreUObject",
+                "DataLink",
+                "Engine",
+            }
+        );
+    }
+}
diff --git a/Engine/Plugins/Experimental/DataLink/Source/DataLinkDataTable/Private/DataLinkDataTableModule.cpp b/Engine/Plugins/Experimental/DataLink/Source/DataLinkDataTable/Private/DataLinkDataTableModule.cpp
new file mode 100644
--- /dev/null
+++ b/Engine/Plugins/Experimental/DataLink/Source/DataLinkDataTable/Private/DataLinkDataTableModule.cpp
@@ -0,0 +1,5 @@
+// Copyright Epic Games, Inc. All Rights Reserved.
+
+#include "Modules/ModuleManager.h"
+
+IMPLEMENT_MODULE(FDefaultModuleImpl, DataLinkDataTable)
diff --git a/Engine/Plugins/Experimental/DataLink/Source/DataLinkDataTable/Private/DataLinkDataTableSource.cpp b/Engine/Plugins/Experimental/DataLink/Source/DataLinkDataTable/Private/DataLinkDataTableSource.cpp
new file mode 100644
--- /dev/null
+++ b/Engine/Plugins/Experimental/DataLink/Source/DataLinkDataTable/Private/DataLinkDataTableSource.cpp
@@ -0,0 +1,43 @@
+// Copyright Epic Games, Inc. All Rights Reserved.
+
+#include "DataLinkDataTableSource.h"
+#include "DataLinkExecutor.h"
+#include "DataLinkNames.h"
+#include "DataLinkNodeInstance.h"
+#include "DataLinkPinBuilder.h"
+#include "StructUtils/StructView.h"
+
+void UDataLinkDataTableSource::OnBuildPins(FDataLinkPinBuilder& Inputs, FDataLinkPinBuilder& Outputs) const
+{
+	Super::OnBuildPins(Inputs, Outputs);
+
+	Inputs.Add(UE::DataLinkDataTable::InputRow)
+		.SetStruct<FDataTableRowHandle>();
+
+	// Output pin does not have a known struct without knowing what the data table to use is
+	Outputs.Add(UE::DataLink::OutputDefault);
+}
+
+EDataLinkExecutionReply UDataLinkDataTableSource::OnExecute(FDataLinkExecutor& InExecutor) const
+{
+	const FDataLinkNodeInstance& NodeInstance = InExecutor.GetNodeInstance(this);
+
+	const FDataLinkInputDataViewer& InputDataViewer = NodeInstance.GetInputDataViewer();
+
+	const FDataTableRowHandle& RowHandle = InputDataViewer.Get<FDataTableRowHandle>(UE::DataLinkDataTable::InputRow);
+	if (!RowHandle.DataTable)
+	{
+		return EDataLinkExecutionReply::Unhandled;
+	}
+
+	const uint8* RowMemory = RowHandle.DataTable->FindRowUnchecked(RowHandle.RowName);
+	if (!RowMemory)
+	{
+		return EDataLinkExecutionReply::Unhandled;
+	}
+
+	// Row memory will return null if the Row struct is null. 
+	check(RowHandle.DataTable->RowStruct != nullptr);
+	InExecutor.SucceedNode(this, FConstStructView(RowHandle.DataTable->RowStruct, RowMemory));
+	return EDataLinkExecutionReply::Handled;
+}
diff --git a/Engine/Plugins/Experimental/DataLink/Source/DataLinkDataTable/Public/DataLinkDataTableSource.h b/Engine/Plugins/Experimental/DataLink/Source/DataLinkDataTable/Public/DataLinkDataTableSource.h
new file mode 100644
--- /dev/null
+++ b/Engine/Plugins/Experimental/DataLink/Source/DataLinkDataTable/Public/DataLinkDataTableSource.h
@@ -0,0 +1,24 @@
+// Copyright Epic Games, Inc. All Rights Reserved.
+
+#pragma once
+
+#include "DataLinkNode.h"
+#include "Engine/DataTable.h"
+#include "DataLinkDataTableSource.generated.h"
+
+namespace UE::DataLinkDataTable
+{
+    const FLazyName InputRow(TEXT("InputRow"));
+}
+
+UCLASS(MinimalAPI, Category="Data Table", DisplayName="Data Table")
+class UDataLinkDataTableSource : public UDataLinkNode
+{
+    GENERATED_BODY()
+
+protected:
+    //~ Begin UDataLinkNode
+    DATALINKDATATABLE_API virtual void OnBuildPins(FDataLinkPinBuilder& Inputs, FDataLinkPinBuilder& Outputs) const override;
+    DATALINKDATATABLE_API virtual EDataLinkExecutionReply OnExecute(FDataLinkExecutor& InExecutor) const override;
+    //~ End UDataLinkNode
+};
diff --git a/Engine/Plugins/Experimental/DataLink/Source/DataLinkEdGraph/DataLinkEdGraph.Build.cs b/Engine/Plugins/Experimental/DataLink/Source/DataLinkEdGraph/DataLinkEdGraph.Build.cs
new file mode 100644
--- /dev/null
+++ b/Engine/Plugins/Experimental/DataLink/Source/DataLinkEdGraph/DataLinkEdGraph.Build.cs
@@ -0,0 +1,39 @@
+// Copyright Epic Games, Inc. All Rights Reserved.
+
+using UnrealBuildTool;
+
+public class DataLinkEdGraph : ModuleRules
+{
+    public DataLinkEdGraph(ReadOnlyTargetRules Target) : base(Target)
+    {
+        PCHUsage = ModuleRules.PCHUsageMode.UseExplicitOrSharedPCHs;
+
+        PublicDependencyModuleNames.AddRange(
+            new string[]
+            {
+                "Core",
+                "CoreUObject",
+                "DataLink",
+                "Engine",
+            }
+        );
+
+        PrivateDependencyModuleNames.AddRange(
+            new string[]
+            {
+                "BlueprintGraph",
+            }
+        );
+
+        if (Target.Type == TargetRules.TargetType.Editor)
+        {
+            PrivateDependencyModuleNames.AddRange(
+                new string[]
+                {
+                    "UnrealEd",
+                    "KismetCompiler",
+                }
+            );
+        }
+    }
+}
diff --git a/Engine/Plugins/Experimental/DataLink/Source/DataLinkEdGraph/Private/Actions/DataLinkGraphAction_NewNativeNode.cpp b/Engine/Plugins/Experimental/DataLink/Source/DataLinkEdGraph/Private/Actions/DataLinkGraphAction_NewNativeNode.cpp
new file mode 100644
--- /dev/null
+++ b/Engine/Plugins/Experimental/DataLink/Source/DataLinkEdGraph/Private/Actions/DataLinkGraphAction_NewNativeNode.cpp
@@ -0,0 +1,35 @@
+// Copyright Epic Games, Inc. All Rights Reserved.
+
+#include "DataLinkGraphAction_NewNativeNode.h"
+#include "DataLinkNode.h"
+#include "Nodes/DataLinkConstant.h"
+
+FDataLinkGraphAction_NewNativeNode::FDataLinkGraphAction_NewNativeNode(TSubclassOf<UDataLinkNode> InNodeClass, int32 InGrouping)
+	: NodeClass(InNodeClass)
+{
+	Grouping = InGrouping;
+
+	UpdateSearchData(InNodeClass->GetDisplayNameText()
+		, InNodeClass->GetToolTipText()
+		, FText::GetEmpty()
+		, FText::GetEmpty());
+}
+
+TSubclassOf<UDataLinkNode> FDataLinkGraphAction_NewNativeNode::GetNodeClass() const
+{
+	return NodeClass;
+}
+
+void FDataLinkGraphAction_NewNativeNode::ConfigureNode(const FConfigContext& InContext) const
+{
+	if (!InContext.SourcePin || !InContext.TemplateNode)
+	{
+		return;
+	}
+
+	if (UDataLinkConstant* ConstantNode = Cast<UDataLinkConstant>(InContext.TemplateNode))
+	{
+		const UScriptStruct* Struct = Cast<UScriptStruct>(InContext.SourcePin->PinType.PinSubCategoryObject);
+		ConstantNode->SetStruct(Struct);
+	}
+}
diff --git a/Engine/Plugins/Experimental/DataLink/Source/DataLinkEdGraph/Private/Actions/DataLinkGraphAction_NewNativeNode.h b/Engine/Plugins/Experimental/DataLink/Source/DataLinkEdGraph/Private/Actions/DataLinkGraphAction_NewNativeNode.h
new file mode 100644
--- /dev/null
+++ b/Engine/Plugins/Experimental/DataLink/Source/DataLinkEdGraph/Private/Actions/DataLinkGraphAction_NewNativeNode.h
@@ -0,0 +1,28 @@
+// Copyright Epic Games, Inc. All Rights Reserved.
+
+#pragma once
+
+#include "DataLinkGraphAction_NewNode.h"
+#include "Templates/SubclassOf.h"
+#include "DataLinkGraphAction_NewNativeNode.generated.h"
+
+class UDataLinkNode;
+
+USTRUCT()
+struct FDataLinkGraphAction_NewNativeNode : public FDataLinkGraphAction_NewNode
+{
+	GENERATED_BODY()
+
+	FDataLinkGraphAction_NewNativeNode() = default;
+
+	explicit FDataLinkGraphAction_NewNativeNode(TSubclassOf<UDataLinkNode> InNodeClass, int32 InGrouping);
+
+protected:
+	//~ Begin FDataLinkGraphAction_NewNode
+	virtual TSubclassOf<UDataLinkNode> GetNodeClass() const override;
+	virtual void ConfigureNode(const FConfigContext& InContext) const override;
+	//~ End FDataLinkGraphAction_NewNode
+
+private:
+	TSubclassOf<UDataLinkNode> NodeClass;
+};
diff --git a/Engine/Plugins/Experimental/DataLink/Source/DataLinkEdGraph/Private/Actions/DataLinkGraphAction_NewNode.cpp b/Engine/Plugins/Experimental/DataLink/Source/DataLinkEdGraph/Private/Actions/DataLinkGraphAction_NewNode.cpp
new file mode 100644
--- /dev/null
+++ b/Engine/Plugins/Experimental/DataLink/Source/DataLinkEdGraph/Private/Actions/DataLinkGraphAction_NewNode.cpp
@@ -0,0 +1,63 @@
+// Copyright Epic Games, Inc. All Rights Reserved.
+
+#include "DataLinkGraphAction_NewNode.h"
+#include "DataLinkNode.h"
+#include "EdGraph/EdGraph.h"
+#include "Nodes/DataLinkEdNode.h"
+#include "ScopedTransaction.h"
+#include "Settings/EditorStyleSettings.h"
+
+#define LOCTEXT_NAMESPACE "DataLinkGraphAction_NewNode"
+
+UEdGraphNode* FDataLinkGraphAction_NewNode::PerformAction(UEdGraph* InParentGraph, UEdGraphPin* InSourcePin, const FVector2f& InLocation, bool bInSelectNewNode)
+{
+	if (!InParentGraph)
+	{
+		return nullptr;
+	}
+
+	TSubclassOf<UDataLinkNode> TemplateNodeClass = GetNodeClass();
+	if (!TemplateNodeClass)
+	{
+		return nullptr;
+	}
+
+	const FScopedTransaction Transaction(LOCTEXT("AddNode", "Add Node"));
+
+	InParentGraph->Modify();
+	if (InSourcePin)
+	{
+		InSourcePin->Modify();
+	}
+
+	UDataLinkEdNode* EdNode = NewObject<UDataLinkEdNode>(InParentGraph);
+	check(EdNode);
+
+	EdNode->SetFlags(RF_Transactional);
+
+	InParentGraph->AddNode(EdNode, /*bUserAction*/true, bInSelectNewNode);
+
+	EdNode->CreateNewGuid();
+	EdNode->PostPlacedNewNode();
+	EdNode->AllocateDefaultPins();
+	EdNode->SetTemplateNodeClass(TemplateNodeClass, /*bReconstructNode*/false);
+
+	{
+		FConfigContext ConfigContext;
+		ConfigContext.SourcePin = InSourcePin;
+		ConfigContext.TemplateNode = EdNode->GetTemplateNode();
+		ConfigureNode(ConfigContext);
+	}
+
+	EdNode->ReconstructNode();
+	EdNode->AutowireNewNode(InSourcePin);
+
+	EdNode->NodePosX = static_cast<int32>(InLocation.X);
+	EdNode->NodePosY = static_cast<int32>(InLocation.Y);
+
+	EdNode->SnapToGrid(GetDefault<UEditorStyleSettings>()->GridSnapSize);
+
+	return EdNode;
+}
+
+#undef LOCTEXT_NAMESPACE
diff --git a/Engine/Plugins/Experimental/DataLink/Source/DataLinkEdGraph/Private/Actions/DataLinkGraphAction_NewNode.h b/Engine/Plugins/Experimental/DataLink/Source/DataLinkEdGraph/Private/Actions/DataLinkGraphAction_NewNode.h
new file mode 100644
--- /dev/null
+++ b/Engine/Plugins/Experimental/DataLink/Source/DataLinkEdGraph/Private/Actions/DataLinkGraphAction_NewNode.h
@@ -0,0 +1,36 @@
+// Copyright Epic Games, Inc. All Rights Reserved.
+
+#pragma once
+
+#include "EdGraph/EdGraphSchema.h"
+#include "Templates/SubclassOf.h"
+#include "DataLinkGraphAction_NewNode.generated.h"
+
+class UDataLinkNode;
+
+USTRUCT()
+struct FDataLinkGraphAction_NewNode : public FEdGraphSchemaAction
+{
+	GENERATED_BODY()
+
+	FDataLinkGraphAction_NewNode() = default;
+
+protected:
+	virtual TSubclassOf<UDataLinkNode> GetNodeClass() const
+	{
+		return nullptr;
+	}
+
+	struct FConfigContext
+	{
+		UDataLinkNode* TemplateNode;
+		UEdGraphPin* SourcePin;
+	};
+	virtual void ConfigureNode(const FConfigContext& InContext) const
+	{
+	}
+
+	//~ Begin FEdGraphSchemaAction
+	virtual UEdGraphNode* PerformAction(UEdGraph* InParentGraph, UEdGraphPin* InSourcePin, const FVector2f& InLocation, bool bInSelectNewNode = true) override;
+	//~ End FEdGraphSchemaAction
+};
diff --git a/Engine/Plugins/Experimental/DataLink/Source/DataLinkEdGraph/Private/Actions/DataLinkGraphAction_NewScriptNode.cpp b/Engine/Plugins/Experimental/DataLink/Source/DataLinkEdGraph/Private/Actions/DataLinkGraphAction_NewScriptNode.cpp
new file mode 100644
--- /dev/null
+++ b/Engine/Plugins/Experimental/DataLink/Source/DataLinkEdGraph/Private/Actions/DataLinkGraphAction_NewScriptNode.cpp
@@ -0,0 +1,69 @@
+// Copyright Epic Games, Inc. All Rights Reserved.
+
+#include "DataLinkGraphAction_NewScriptNode.h"
+#include "Blueprint/BlueprintSupport.h"
+#include "Engine/Blueprint.h"
+#include "Nodes/DataLinkEdNode.h"
+#include "Nodes/Script/DataLinkScriptNode.h"
+#include "Nodes/Script/DataLinkScriptNodeWrapper.h"
+#include "UObject/Class.h"
+
+#define LOCTEXT_NAMESPACE "DataLinkGraphAction_NewScriptNode"
+
+FDataLinkGraphAction_NewScriptNode::FDataLinkGraphAction_NewScriptNode(const FAssetData& InNodeAsset, int32 InGrouping)
+	: NodeAsset(InNodeAsset)
+{
+	Grouping = InGrouping;
+
+	auto GetTag = [&InNodeAsset](FName InTagKey, const FText& InDefaultValue = FText::GetEmpty())
+		{
+			FString TagValue;
+			if (InNodeAsset.GetTagValue(InTagKey, TagValue))
+			{
+				return FText::FromString(TagValue);
+			}
+			return InDefaultValue;
+		};
+
+	const FText AssetNameText = FText::FromName(InNodeAsset.AssetName);
+	UpdateSearchData(GetTag(FBlueprintTags::BlueprintDisplayName, AssetNameText)
+		, GetTag(FBlueprintTags::BlueprintDescription)
+		, GetTag(FBlueprintTags::BlueprintCategory, LOCTEXT("DefaultBlueprintCategory", "Script Nodes"))
+		, FText::GetEmpty());
+}
+
+TSubclassOf<UDataLinkNode> FDataLinkGraphAction_NewScriptNode::GetNodeClass() const
+{
+	return UDataLinkScriptNodeWrapper::StaticClass();
+}
+
+void FDataLinkGraphAction_NewScriptNode::ConfigureNode(const FConfigContext& InContext) const
+{
+	if (UDataLinkScriptNodeWrapper* ScriptNodeWrapper = Cast<UDataLinkScriptNodeWrapper>(InContext.TemplateNode))
+	{
+		ScriptNodeWrapper->SetNodeClass(GetScriptNodeClass());
+	}
+}
+
+TSubclassOf<UDataLinkScriptNode> FDataLinkGraphAction_NewScriptNode::GetScriptNodeClass() const
+{
+	UObject* NodeObject = NodeAsset.GetAsset();
+	if (!NodeObject)
+	{
+		return nullptr;
+	}
+
+	if (TSubclassOf<UDataLinkScriptNode> NodeClass = Cast<UClass>(NodeObject))
+	{
+		return NodeClass;
+	}
+
+	if (UBlueprint* Blueprint = Cast<UBlueprint>(NodeObject))
+	{
+		return Cast<UClass>(Blueprint->GeneratedClass);
+	}
+
+	return nullptr;
+}
+
+#undef LOCTEXT_NAMESPACE
diff --git a/Engine/Plugins/Experimental/DataLink/Source/DataLinkEdGraph/Private/Actions/DataLinkGraphAction_NewScriptNode.h b/Engine/Plugins/Experimental/DataLink/Source/DataLinkEdGraph/Private/Actions/DataLinkGraphAction_NewScriptNode.h
new file mode 100644
--- /dev/null
+++ b/Engine/Plugins/Experimental/DataLink/Source/DataLinkEdGraph/Private/Actions/DataLinkGraphAction_NewScriptNode.h
@@ -0,0 +1,28 @@
+// Copyright Epic Games, Inc. All Rights Reserved.
+
+#pragma once
+
+#include "DataLinkGraphAction_NewNode.h"
+#include "Nodes/Script/DataLinkScriptNode.h"
+#include "DataLinkGraphAction_NewScriptNode.generated.h"
+
+USTRUCT()
+struct FDataLinkGraphAction_NewScriptNode : public FDataLinkGraphAction_NewNode
+{
+	GENERATED_BODY()
+
+	FDataLinkGraphAction_NewScriptNode() = default;
+
+	explicit FDataLinkGraphAction_NewScriptNode(const FAssetData& InNodeAsset, int32 InGrouping);
+
+protected:
+	//~ Begin FDataLinkGraphAction_NewNode
+	virtual TSubclassOf<UDataLinkNode> GetNodeClass() const override;
+	virtual void ConfigureNode(const FConfigContext& InContext) const override;
+	//~ End FDataLinkGraphAction_NewNode
+
+private:
+	TSubclassOf<UDataLinkScriptNode> GetScriptNodeClass() const;
+
+	FAssetData NodeAsset;
+};
diff --git a/Engine/Plugins/Experimental/DataLink/Source/DataLinkEdGraph/Private/Blueprint/K2Node_AsyncDataLinkRequest.cpp b/Engine/Plugins/Experimental/DataLink/Source/DataLinkEdGraph/Private/Blueprint/K2Node_AsyncDataLinkRequest.cpp
new file mode 100644
--- /dev/null
+++ b/Engine/Plugins/Experimental/DataLink/Source/DataLinkEdGraph/Private/Blueprint/K2Node_AsyncDataLinkRequest.cpp
@@ -0,0 +1,46 @@
+// Copyright Epic Games, Inc. All Rights Reserved.
+
+#include "K2Node_AsyncDataLinkRequest.h"
+#include "DataLinkRequestProxy.h"
+
+UK2Node_AsyncDataLinkRequest::UK2Node_AsyncDataLinkRequest()
+{
+	ProxyFactoryFunctionName = GET_FUNCTION_NAME_CHECKED(UDataLinkRequestProxy, CreateRequestProxy);
+	ProxyFactoryClass = UDataLinkRequestProxy::StaticClass();
+	ProxyClass = UDataLinkRequestProxy::StaticClass();
+}
+
+UEdGraphPin* UK2Node_AsyncDataLinkRequest::GetDataLinkInstancePin() const
+{
+	return FindPinChecked(TEXT("InDataLinkInstance"), EGPD_Input);
+}
+
+UEdGraphPin* UK2Node_AsyncDataLinkRequest::GetExecutionContextPin() const
+{
+	return FindPinChecked(TEXT("InExecutionContext"), EGPD_Input);
+}
+
+UEdGraphPin* UK2Node_AsyncDataLinkRequest::GetDataLinkSinkProviderPin() const
+{
+	return FindPinChecked(TEXT("InDataLinkSinkProvider"), EGPD_Input);
+}
+
+UEdGraphPin* UK2Node_AsyncDataLinkRequest::GetOutputDataPin() const
+{
+	return FindPinChecked(TEXT("OutputData"), EGPD_Output);
+}
+
+UEdGraphPin* UK2Node_AsyncDataLinkRequest::GetExecutionResultPin() const
+{
+	return FindPinChecked(TEXT("ExecutionResult"), EGPD_Output);
+}
+
+UEdGraphPin* UK2Node_AsyncDataLinkRequest::GetRequestCompletePin() const
+{
+	return FindPinChecked(TEXT("OnRequestComplete"), EGPD_Output);
+}
+
+void UK2Node_AsyncDataLinkRequest::GetMenuActions(FBlueprintActionDatabaseRegistrar& InActionRegistrar) const
+{
+	// Hide this node from menu actions
+}
diff --git a/Engine/Plugins/Experimental/DataLink/Source/DataLinkEdGraph/Private/Blueprint/K2Node_AsyncDataLinkRequest.h b/Engine/Plugins/Experimental/DataLink/Source/DataLinkEdGraph/Private/Blueprint/K2Node_AsyncDataLinkRequest.h
new file mode 100644
--- /dev/null
+++ b/Engine/Plugins/Experimental/DataLink/Source/DataLinkEdGraph/Private/Blueprint/K2Node_AsyncDataLinkRequest.h
@@ -0,0 +1,27 @@
+// Copyright Epic Games, Inc. All Rights Reserved.
+
+#pragma once
+
+#include "K2Node_BaseAsyncTask.h"
+#include "K2Node_AsyncDataLinkRequest.generated.h"
+
+UCLASS()
+class UK2Node_AsyncDataLinkRequest : public UK2Node_BaseAsyncTask
+{
+	GENERATED_BODY()
+
+public:
+	UK2Node_AsyncDataLinkRequest();
+
+	UEdGraphPin* GetDataLinkInstancePin() const;
+	UEdGraphPin* GetExecutionContextPin() const;
+	UEdGraphPin* GetDataLinkSinkProviderPin() const;
+	UEdGraphPin* GetOutputDataPin() const;
+	UEdGraphPin* GetExecutionResultPin() const;
+	UEdGraphPin* GetRequestCompletePin() const;
+
+protected:
+	//~ Begin UK2Node
+	virtual void GetMenuActions(FBlueprintActionDatabaseRegistrar& InActionRegistrar) const override;
+	//~ End UK2Node
+};
diff --git a/Engine/Plugins/Experimental/DataLink/Source/DataLinkEdGraph/Private/Blueprint/K2Node_DataLinkRequest.cpp b/Engine/Plugins/Experimental/DataLink/Source/DataLinkEdGraph/Private/Blueprint/K2Node_DataLinkRequest.cpp
new file mode 100644
--- /dev/null
+++ b/Engine/Plugins/Experimental/DataLink/Source/DataLinkEdGraph/Private/Blueprint/K2Node_DataLinkRequest.cpp
@@ -0,0 +1,166 @@
+// Copyright Epic Games, Inc. All Rights Reserved.
+
+#include "Blueprint/K2Node_DataLinkRequest.h"
+#include "BlueprintActionDatabaseRegistrar.h"
+#include "BlueprintNodeSpawner.h"
+#include "DataLinkEnums.h"
+#include "DataLinkGraph.h"
+#include "DataLinkInstance.h"
+#include "IDataLinkSinkProvider.h"
+#include "K2Node_AsyncDataLinkRequest.h"
+#include "K2Node_SwitchEnum.h"
+#include "KismetCompiler.h"
+#include "Nodes/DataLinkEdNode.h"
+#include "StructUtils/InstancedStruct.h"
+
+#define LOCTEXT_NAMESPACE "K2Node_DataLinkRequest"
+
+const FLazyName UK2Node_DataLinkRequest::PN_DataLinkInstance = TEXT("DataLinkInstance");
+const FLazyName UK2Node_DataLinkRequest::PN_ExecutionContext = TEXT("ExecutionContext");
+const FLazyName UK2Node_DataLinkRequest::PN_DataLinkSinkProvider = TEXT("DataLinkSinkProvider");
+const FLazyName UK2Node_DataLinkRequest::PN_Success = TEXT("Success");
+const FLazyName UK2Node_DataLinkRequest::PN_Failure = TEXT("Failure");
+
+UEdGraphPin* UK2Node_DataLinkRequest::GetDataLinkInstancePin() const
+{
+	return FindPinChecked(UK2Node_DataLinkRequest::PN_DataLinkInstance, EGPD_Input);
+}
+
+UEdGraphPin* UK2Node_DataLinkRequest::GetExecutionContextPin() const
+{
+	return FindPinChecked(UK2Node_DataLinkRequest::PN_ExecutionContext, EGPD_Input);
+}
+
+UEdGraphPin* UK2Node_DataLinkRequest::GetDataLinkSinkProviderPin() const
+{
+	return FindPinChecked(UK2Node_DataLinkRequest::PN_DataLinkSinkProvider, EGPD_Input);
+}
+
+UEdGraphPin* UK2Node_DataLinkRequest::GetOutputPin() const
+{
+	return FindPinChecked(UDataLinkEdNode::PN_Output, EGPD_Output);
+}
+
+void UK2Node_DataLinkRequest::AllocateDefaultPins()
+{
+	Super::AllocateDefaultPins();
+
+	// Exec Pin
+	CreatePin(EGPD_Input, UEdGraphSchema_K2::PC_Exec, UEdGraphSchema_K2::PN_Execute);
+
+	// Then Pin
+	UEdGraphPin* const ThenPin = CreatePin(EGPD_Output, UEdGraphSchema_K2::PC_Exec, UEdGraphSchema_K2::PN_Then);
+	ThenPin->PinFriendlyName = LOCTEXT("ThenPinLabel", "Processing");
+	ThenPin->PinToolTip = TEXT("For execution that needs to happen immediately without waiting for a response.");
+
+	// Success Pin
+	UEdGraphPin* const SuccessPin = CreatePin(EGPD_Output, UEdGraphSchema_K2::PC_Exec, UK2Node_DataLinkRequest::PN_Success);
+	SuccessPin->PinFriendlyName = LOCTEXT("SuccessPinLabel", "Response Received");
+	SuccessPin->PinToolTip = TEXT("Executes only when the request was successful and received the data");
+
+	// Failure pin
+	UEdGraphPin* const FailurePin = CreatePin(EGPD_Output, UEdGraphSchema_K2::PC_Exec, UK2Node_DataLinkRequest::PN_Failure);
+	FailurePin->PinFriendlyName = LOCTEXT("FailurePinLabel", "Request Failed");
+	FailurePin->PinToolTip = TEXT("Executes when the request failed and no data will be received");
+
+	// Create Data Link Instance Pin
+	CreatePin(EGPD_Input, UEdGraphSchema_K2::PC_Struct, FDataLinkInstance::StaticStruct(), UK2Node_DataLinkRequest::PN_DataLinkInstance);
+
+	// Create Execution Context Pin
+	if (UEdGraphPin* ExecutionContextPin = CreatePin(EGPD_Input, UEdGraphSchema_K2::PC_Object, UObject::StaticClass(), UK2Node_DataLinkRequest::PN_ExecutionContext))
+	{
+		// Prevent user from editing default value (i.e. hiding the object dropdown as option).
+		ExecutionContextPin->bDefaultValueIsIgnored = true;
+	}
+
+	// Create Data Link Sink Provider Pin
+	if (UEdGraphPin* SinkProviderPin = CreatePin(EGPD_Input, UEdGraphSchema_K2::PC_Interface, UDataLinkSinkProvider::StaticClass(), UK2Node_DataLinkRequest::PN_DataLinkSinkProvider))
+	{
+		// Prevent user from editing default value (i.e. hiding the object dropdown as option).
+		SinkProviderPin->bDefaultValueIsIgnored = true;
+	}
+
+	// Create Output Pin
+	CreatePin(EGPD_Output, UEdGraphSchema_K2::PC_Struct, FInstancedStruct::StaticStruct(), UDataLinkEdNode::PN_Output);
+}
+
+bool UK2Node_DataLinkRequest::IsCompatibleWithGraph(const UEdGraph* InTargetGraph) const
+{
+	const UEdGraphSchema* Schema = InTargetGraph->GetSchema();
+	if (!Schema)
+	{
+		return false;
+	}
+
+	const EGraphType GraphType = Schema->GetGraphType(InTargetGraph);
+	return (GraphType == GT_Ubergraph || GraphType == GT_Macro) && Super::IsCompatibleWithGraph(InTargetGraph);
+}
+
+FText UK2Node_DataLinkRequest::GetNodeTitle(ENodeTitleType::Type InTitleType) const
+{
+	return LOCTEXT("NodeTitle", "Data Link Request");
+}
+
+void UK2Node_DataLinkRequest::ExpandNode(FKismetCompilerContext& InCompilerContext, UEdGraph* InSourceGraph)
+{
+	UK2Node_AsyncDataLinkRequest* AsyncRequestNode = InCompilerContext.SpawnIntermediateNode<UK2Node_AsyncDataLinkRequest>(this, InSourceGraph);
+	AsyncRequestNode->AllocateDefaultPins();
+
+	InCompilerContext.MovePinLinksToIntermediate(*GetExecPin(), *AsyncRequestNode->GetExecPin());
+	InCompilerContext.MovePinLinksToIntermediate(*GetThenPin(), *AsyncRequestNode->GetThenPin());
+	InCompilerContext.MovePinLinksToIntermediate(*GetDataLinkInstancePin(), *AsyncRequestNode->GetDataLinkInstancePin());
+	InCompilerContext.MovePinLinksToIntermediate(*GetExecutionContextPin(), *AsyncRequestNode->GetExecutionContextPin());
+	InCompilerContext.MovePinLinksToIntermediate(*GetDataLinkSinkProviderPin(), *AsyncRequestNode->GetDataLinkSinkProviderPin());
+	InCompilerContext.MovePinLinksToIntermediate(*GetOutputPin(), *AsyncRequestNode->GetOutputDataPin());
+
+	UEnum* const ResultEnum = StaticEnum<EDataLinkExecutionResult>();
+
+	UK2Node_SwitchEnum* SwitchNode = InCompilerContext.SpawnIntermediateNode<UK2Node_SwitchEnum>(this, InSourceGraph);
+	SwitchNode->Enum = ResultEnum;
+	SwitchNode->AllocateDefaultPins();
+
+	// Connect Request and Result Enum value pins to the Switch Node
+	{
+		UEdGraphPin* RequestCompletePin = AsyncRequestNode->GetRequestCompletePin();
+		UEdGraphPin* ExecutionResultPin = AsyncRequestNode->GetExecutionResultPin();
+		check(RequestCompletePin && ExecutionResultPin);
+
+		RequestCompletePin->MakeLinkTo(SwitchNode->GetExecPin());
+		ExecutionResultPin->MakeLinkTo(SwitchNode->GetSelectionPin());
+	}
+
+	// Move Result Exec pins
+	{
+		UEdGraphPin* SuccessOutputPin = SwitchNode->FindPinChecked(ResultEnum->GetNameStringByValue(static_cast<uint64>(EDataLinkExecutionResult::Succeeded)), EGPD_Output);
+		UEdGraphPin* FailureOutputPin = SwitchNode->FindPinChecked(ResultEnum->GetNameStringByValue(static_cast<uint64>(EDataLinkExecutionResult::Failed)), EGPD_Output);
+
+		InCompilerContext.MovePinLinksToIntermediate(*FindPinChecked(UK2Node_DataLinkRequest::PN_Success), *SuccessOutputPin);
+		InCompilerContext.MovePinLinksToIntermediate(*FindPinChecked(UK2Node_DataLinkRequest::PN_Failure), *FailureOutputPin);
+	}
+
+	BreakAllNodeLinks();
+}
+
+FName UK2Node_DataLinkRequest::GetCornerIcon() const
+{
+	return TEXT("Graph.Latent.LatentIcon");
+}
+
+void UK2Node_DataLinkRequest::GetMenuActions(FBlueprintActionDatabaseRegistrar& InActionRegistrar) const
+{
+	UClass* ActionKey = GetClass();
+
+	if (InActionRegistrar.IsOpenForRegistration(ActionKey))
+	{
+		UBlueprintNodeSpawner* NodeSpawner = UBlueprintNodeSpawner::Create(ActionKey);
+		checkf(NodeSpawner, TEXT("Node spawner failed to create a valid Node"));
+		InActionRegistrar.AddBlueprintAction(ActionKey, NodeSpawner);
+	}
+}
+
+FText UK2Node_DataLinkRequest::GetMenuCategory() const
+{
+	return LOCTEXT("MenuCategory", "Data Link");
+}
+
+#undef LOCTEXT_NAMESPACE
diff --git a/Engine/Plugins/Experimental/DataLink/Source/DataLinkEdGraph/Private/DataLinkEdGraph.cpp b/Engine/Plugins/Experimental/DataLink/Source/DataLinkEdGraph/Private/DataLinkEdGraph.cpp
new file mode 100644
--- /dev/null
+++ b/Engine/Plugins/Experimental/DataLink/Source/DataLinkEdGraph/Private/DataLinkEdGraph.cpp
@@ -0,0 +1,70 @@
+// Copyright Epic Games, Inc. All Rights Reserved.
+
+#include "DataLinkEdGraph.h"
+#include "DataLinkGraph.h"
+#include "Nodes/DataLinkEdOutputNode.h"
+
+UDataLinkEdGraph::UDataLinkEdGraph()
+{
+	OnGraphCompiledHandle = UDataLinkGraph::OnGraphCompiled().AddUObject(this, &UDataLinkEdGraph::OnGraphCompiled);
+}
+
+UDataLinkEdOutputNode* UDataLinkEdGraph::FindOutputNode() const
+{
+	UDataLinkEdOutputNode* OutputNode;
+	if (Nodes.FindItemByClass(&OutputNode))
+	{
+		return OutputNode;
+	}
+	return nullptr;
+}
+
+void UDataLinkEdGraph::InitializeNodes()
+{
+	for (UEdGraphNode* Node : Nodes)
+	{
+		UDataLinkEdNode* DataLinkEdNode = Cast<UDataLinkEdNode>(Node);
+		if (!DataLinkEdNode)
+		{
+			continue;
+		}
+
+		DataLinkEdNode->UpdateMetadata();
+
+		if (DataLinkEdNode->RequiresPinRecreation())
+		{
+			DataLinkEdNode->RecreatePins();	
+		}
+	}
+}
+
+void UDataLinkEdGraph::BeginDestroy()
+{
+	Super::BeginDestroy();
+
+	UDataLinkGraph::OnGraphCompiled().Remove(OnGraphCompiledHandle);
+	OnGraphCompiledHandle.Reset();
+}
+
+void UDataLinkEdGraph::DirtyGraph()
+{
+	// ChangeId could be regenerated every time graph is dirtied,
+	// but it only makes sense when it matches the compiled change id
+	if (ChangeId == LastCompiledChangeId)
+	{
+		ChangeId = FGuid::NewGuid();
+	}
+}
+
+bool UDataLinkEdGraph::IsCompiledGraphUpToDate() const
+{
+	return ChangeId == LastCompiledChangeId;
+}
+
+void UDataLinkEdGraph::OnGraphCompiled(UDataLinkGraph* InCompiledGraph)
+{
+	if (InCompiledGraph && InCompiledGraph->GetEdGraph() == this)
+	{
+		LastCompiledChangeId = ChangeId;
+	}
+}
diff --git a/Engine/Plugins/Experimental/DataLink/Source/DataLinkEdGraph/Private/DataLinkEdGraphLog.h b/Engine/Plugins/Experimental/DataLink/Source/DataLinkEdGraph/Private/DataLinkEdGraphLog.h
new file mode 100644
--- /dev/null
+++ b/Engine/Plugins/Experimental/DataLink/Source/DataLinkEdGraph/Private/DataLinkEdGraphLog.h
@@ -0,0 +1,7 @@
+// Copyright Epic Games, Inc. All Rights Reserved.
+
+#pragma once
+
+#include "Logging/LogMacros.h"
+
+DECLARE_LOG_CATEGORY_EXTERN(LogDataLinkEdGraph, Log, All);
diff --git a/Engine/Plugins/Experimental/DataLink/Source/DataLinkEdGraph/Private/DataLinkEdGraphModule.cpp b/Engine/Plugins/Experimental/DataLink/Source/DataLinkEdGraph/Private/DataLinkEdGraphModule.cpp
new file mode 100644
--- /dev/null
+++ b/Engine/Plugins/Experimental/DataLink/Source/DataLinkEdGraph/Private/DataLinkEdGraphModule.cpp
@@ -0,0 +1,8 @@
+// Copyright Epic Games, Inc. All Rights Reserved.
+
+#include "DataLinkEdGraphLog.h"
+#include "Modules/ModuleManager.h"
+
+DEFINE_LOG_CATEGORY(LogDataLinkEdGraph);
+
+IMPLEMENT_MODULE(FDefaultModuleImpl, DataLinkEdGraph)
diff --git a/Engine/Plugins/Experimental/DataLink/Source/DataLinkEdGraph/Private/DataLinkEdGraphSchema.cpp b/Engine/Plugins/Experimental/DataLink/Source/DataLinkEdGraph/Private/DataLinkEdGraphSchema.cpp
new file mode 100644
--- /dev/null
+++ b/Engine/Plugins/Experimental/DataLink/Source/DataLinkEdGraph/Private/DataLinkEdGraphSchema.cpp
@@ -0,0 +1,220 @@
+// Copyright Epic Games, Inc. All Rights Reserved.
+
+#include "DataLinkEdGraphSchema.h"
+#include "Actions/DataLinkGraphAction_NewNativeNode.h"
+#include "Actions/DataLinkGraphAction_NewScriptNode.h"
+#include "AssetRegistry/AssetRegistryModule.h"
+#include "AssetRegistry/IAssetRegistry.h"
+#include "Blueprint/BlueprintSupport.h"
+#include "DataLinkNode.h"
+#include "EdGraph/EdGraph.h"
+#include "Nodes/DataLinkEdNode.h"
+#include "Nodes/DataLinkEdOutputNode.h"
+#include "Nodes/Script/DataLinkScriptNode.h"
+#include "UObject/Class.h"
+#include "UObject/UObjectIterator.h"
+
+#define LOCTEXT_NAMESPACE "DataLinkEdGraphSchema"
+
+/** Pin Categories */
+const FLazyName UDataLinkEdGraphSchema::PC_Data = TEXT("Data");
+
+/** Pin Category Colors */
+const FLinearColor UDataLinkEdGraphSchema::PCC_Data = FLinearColor::White;
+
+namespace UE::DataLinkEdGraph::Private
+{
+	template<typename InActionType, typename... InArgTypes>
+	TSharedRef<InActionType> AddAction(FGraphContextMenuBuilder& InContextMenuBuilder, const FString& InCategory, InArgTypes&&... InArgs)
+	{
+		TSharedRef<InActionType> Action = MakeShared<InActionType>(Forward<InArgTypes>(InArgs)...);
+		Action->CosmeticUpdateRootCategory(FText::FromString(InCategory));
+		InContextMenuBuilder.AddAction(Action);
+		return Action;
+	}
+}
+
+bool UDataLinkEdGraphSchema::IsConnectionLooping(const UEdGraphPin* InInputPin, const UEdGraphPin* InOutputPin) const
+{ 
+	// Inline size of 1 as most cases the nodes will be connected 1 to 1
+	TArray<const UEdGraphNode*, TInlineAllocator<1>> NodesToCheck = { InInputPin->GetOwningNode() };
+
+	const UEdGraphNode* OutputNode = InOutputPin->GetOwningNode();
+
+	// Loop by starting with the Input Pin's Node and going in the direction opposite to the Input Pin direction (Output)
+	// and see if we encounter the Output Pin's node
+	while (!NodesToCheck.IsEmpty())
+	{
+		if (const UEdGraphNode* Node = NodesToCheck.Pop(EAllowShrinking::No))
+		{
+			if (Node == OutputNode)
+			{
+				// Output node detected! Looping found
+				return true;
+			}
+
+			for (const UEdGraphPin* Pin : Node->Pins)
+			{
+				if (Pin && Pin->Direction == EGPD_Output)
+				{
+					for (const UEdGraphPin* LinkedPin : Pin->LinkedTo)
+					{
+						NodesToCheck.AddUnique(LinkedPin->GetOwningNode());
+					}
+				}
+			}
+		}
+	}
+	return false;
+}
+
+void UDataLinkEdGraphSchema::CreateDefaultNodesForGraph(UEdGraph& InGraph) const
+{
+	FGraphNodeCreator<UDataLinkEdOutputNode> NodeCreator(InGraph);
+	UDataLinkEdOutputNode* OutputNode = NodeCreator.CreateNode();
+	NodeCreator.Finalize();
+
+	SetNodeMetaData(OutputNode, FNodeMetadata::DefaultGraphNode);
+}
+
+bool UDataLinkEdGraphSchema::ArePinsCompatible(const UEdGraphPin* InPinA, const UEdGraphPin* InPinB, const UClass* InCallingContext, bool bInIgnoreArray) const
+{
+	if (!InPinA || !InPinB)
+	{
+		return false;
+	}
+
+	auto IsDirectionCompatible = [](const UEdGraphPin* InPinA, const UEdGraphPin* InPinB)
+		{
+			return InPinA->Direction == EGPD_Input && InPinB->Direction == EGPD_Output;
+		};
+
+	if (!IsDirectionCompatible(InPinA, InPinB) && !IsDirectionCompatible(InPinB, InPinA))
+	{
+		return false;
+	}
+
+	return InPinA->PinType.PinCategory == InPinB->PinType.PinCategory
+		&& InPinA->PinType.PinSubCategoryObject == InPinB->PinType.PinSubCategoryObject;
+}
+
+void UDataLinkEdGraphSchema::GetGraphContextActions(FGraphContextMenuBuilder& InContextMenuBuilder) const
+{
+	using namespace UE::DataLinkEdGraph;
+
+	constexpr int32 Grouping = 0; 
+
+	IAssetRegistry& AssetRegistry = FModuleManager::LoadModuleChecked<FAssetRegistryModule>(FName("AssetRegistry")).Get();
+
+	const FName MD_Category = TEXT("Category");
+
+	// Add all node assets
+	{
+		const TMultiMap<FName, FString> TagValues =
+		{
+			{ FBlueprintTags::NativeParentClassPath, FObjectPropertyBase::GetExportPath(UDataLinkScriptNode::StaticClass()) },
+		};
+
+		TArray<FAssetData> ScriptNodeAssets;
+		AssetRegistry.GetAssetsByTagValues(TagValues, ScriptNodeAssets);
+
+		for (const FAssetData& ScriptNodeAsset : ScriptNodeAssets)
+		{
+			Private::AddAction<FDataLinkGraphAction_NewScriptNode>(InContextMenuBuilder
+				, ScriptNodeAsset.GetTagValueRef<FString>(MD_Category)
+				, ScriptNodeAsset
+				, Grouping);
+		}
+	}
+
+	const FName MD_Hidden = TEXT("Hidden");
+
+	// Add all node native classes
+	for (UClass* Class : TObjectRange<UClass>())
+	{
+		if (Class->HasAnyClassFlags(CLASS_Abstract | CLASS_Deprecated) || Class->HasMetaData(MD_Hidden))
+		{
+			continue;
+		}
+
+		if (Class->HasAllClassFlags(CLASS_Native) && Class->IsChildOf<UDataLinkNode>())
+		{
+			Private::AddAction<FDataLinkGraphAction_NewNativeNode>(InContextMenuBuilder
+				, Class->GetMetaData(MD_Category)
+				, TSubclassOf<UDataLinkNode>(Class)
+				, Grouping);
+		}
+	}
+}
+
+void UDataLinkEdGraphSchema::GetContextMenuActions(UToolMenu* InMenu, UGraphNodeContextMenuContext* InContext) const
+{
+	Super::GetContextMenuActions(InMenu, InContext);
+}
+
+const FPinConnectionResponse UDataLinkEdGraphSchema::CanCreateConnection(const UEdGraphPin* InSourcePin, const UEdGraphPin* InTargetPin) const
+{
+	// Make sure the pins are not on the same node
+	if (InSourcePin->GetOwningNode() == InTargetPin->GetOwningNode())
+	{
+		return FPinConnectionResponse(CONNECT_RESPONSE_DISALLOW, LOCTEXT("ConnectionSameNode", "Both are on the same node"));
+	}
+
+	// Pin mismatch in Pin Category
+	if (InSourcePin->PinType.PinCategory != InTargetPin->PinType.PinCategory)
+	{
+		return FPinConnectionResponse(CONNECT_RESPONSE_DISALLOW, LOCTEXT("IncompatibleCategories", "Pin Types are not Compatible"));
+	}
+
+	// Compare the directions
+	const UEdGraphPin* InputPin = nullptr;
+	const UEdGraphPin* OutputPin = nullptr;
+
+	if (!CategorizePinsByDirection(InSourcePin, InTargetPin, InputPin, OutputPin))
+	{
+		return FPinConnectionResponse(CONNECT_RESPONSE_DISALLOW, LOCTEXT("IncompatibleDirections", "Directions are not compatible"));
+	}
+
+	if (IsConnectionLooping(InputPin, OutputPin))
+	{
+		return FPinConnectionResponse(CONNECT_RESPONSE_DISALLOW, LOCTEXT("ConnectionLoop", "Connection would cause loop"));
+	}
+
+	const bool bHasSourcePinLinks = !InSourcePin->LinkedTo.IsEmpty();
+	const bool bHasTargetPinLinks = !InTargetPin->LinkedTo.IsEmpty();
+
+	if (bHasSourcePinLinks && bHasTargetPinLinks)
+	{
+		return FPinConnectionResponse(CONNECT_RESPONSE_BREAK_OTHERS_AB, LOCTEXT("ConnectionReplaceAB", "Replace existing connections"));
+	}
+
+	if (bHasSourcePinLinks)
+	{
+		return FPinConnectionResponse(CONNECT_RESPONSE_BREAK_OTHERS_A, LOCTEXT("ConnectionReplaceA", "Replace existing connections"));
+	}
+
+	if (bHasTargetPinLinks)
+	{
+		return FPinConnectionResponse(CONNECT_RESPONSE_BREAK_OTHERS_B, LOCTEXT("ConnectionReplaceB", "Replace existing connections"));
+	}
+
+	return FPinConnectionResponse(CONNECT_RESPONSE_MAKE, TEXT(""));
+}
+
+FLinearColor UDataLinkEdGraphSchema::GetPinTypeColor(const FEdGraphPinType& InPinType) const
+{
+	if (InPinType.PinCategory == PC_Data)
+	{
+		return UDataLinkEdGraphSchema::PCC_Data;
+	}
+	return Super::GetPinTypeColor(InPinType);
+}
+
+void UDataLinkEdGraphSchema::GetGraphDisplayInformation(const UEdGraph& InGraph, FGraphDisplayInfo& OutDisplayInfo) const
+{
+	OutDisplayInfo.PlainName = FText::FromName(InGraph.GetFName());
+	OutDisplayInfo.DisplayName = OutDisplayInfo.PlainName;
+	OutDisplayInfo.Tooltip = LOCTEXT("GraphTooltip", "Graph used to determine how data is linked and flows from a source");
+}
+
+#undef LOCTEXT_NAMESPACE
diff --git a/Engine/Plugins/Experimental/DataLink/Source/DataLinkEdGraph/Private/Nodes/DataLinkEdNode.cpp b/Engine/Plugins/Experimental/DataLink/Source/DataLinkEdGraph/Private/Nodes/DataLinkEdNode.cpp
new file mode 100644
--- /dev/null
+++ b/Engine/Plugins/Experimental/DataLink/Source/DataLinkEdGraph/Private/Nodes/DataLinkEdNode.cpp
@@ -0,0 +1,319 @@
+// Copyright Epic Games, Inc. All Rights Reserved.
+
+#include "Nodes/DataLinkEdNode.h"
+#include "DataLinkEdGraphLog.h"
+#include "DataLinkEdGraphSchema.h"
+#include "DataLinkNode.h"
+#include "DataLinkUtils.h"
+#include "EdGraph/EdGraph.h"
+#include "Engine/Engine.h"
+#include "Templates/SubclassOf.h"
+
+const FLazyName UDataLinkEdNode::MD_InvalidatesNode = TEXT("InvalidatesNode");
+const FLazyName UDataLinkEdNode::PN_Output = TEXT("Output");
+
+namespace UE::DataLinkEdGraph::Private
+{
+	bool FindLinkedNode(const UEdGraphPin* InPin, const UDataLinkEdNode*& OutLinkedNode, const UEdGraphPin*& OutLinkedPin)
+	{
+		if (!InPin || InPin->LinkedTo.IsEmpty())
+		{
+			return false;
+		}
+
+		UEdGraphPin* LinkedPin = InPin->LinkedTo[0];
+		if (!LinkedPin)
+		{
+			return false;
+		}
+
+		if (const UDataLinkEdNode* EdNode = Cast<UDataLinkEdNode>(LinkedPin->GetOwningNode()))
+		{
+			OutLinkedNode = EdNode;
+			OutLinkedPin = LinkedPin;
+			return true;
+		}
+		return false;
+	}
+
+	int32 GetDataPinCount(const UDataLinkEdNode& InEdNode)
+	{
+		int32 DataPinCount = 0;
+		for (UEdGraphPin* Pin : InEdNode.Pins)
+		{
+			if (Pin && Pin->PinType.PinCategory == UDataLinkEdGraphSchema::PC_Data)
+			{
+				++DataPinCount;
+			}
+		}
+		return DataPinCount;
+	}
+}
+
+void UDataLinkEdNode::SetTemplateNodeClass(TSubclassOf<UDataLinkNode> InNodeClass, bool bInReconstructNode)
+{
+	UObject* Node = TemplateNode;
+	if (UE::DataLink::ReplaceObject(Node, this, InNodeClass))
+	{
+		TemplateNode = Cast<UDataLinkNode>(Node);
+
+		if (bInReconstructNode)
+		{
+			ReconstructNode();
+		}
+	}
+}
+
+void UDataLinkEdNode::ForEachPinConnection(TFunctionRef<void(const UEdGraphPin&, const UDataLinkEdNode&, const UEdGraphPin&)> InFunction) const
+{
+	for (const UEdGraphPin* Pin : Pins)
+	{
+		const UDataLinkEdNode* LinkedNode;
+		const UEdGraphPin* LinkedPin;
+
+		if (UE::DataLinkEdGraph::Private::FindLinkedNode(Pin, LinkedNode, LinkedPin))
+		{
+			InFunction(*Pin, *LinkedNode, *LinkedPin);
+		}
+	}
+}
+
+bool UDataLinkEdNode::RequiresPinRecreation() const
+{
+	const int32 DataPinCount = UE::DataLinkEdGraph::Private::GetDataPinCount(*this);
+	if (!TemplateNode)
+	{
+		// If a template node is not set, but there are still Data Pins in place, they need to be cleared off
+		return DataPinCount != 0;
+	}
+
+	TArray<FDataLinkPin> InputPins;
+	TArray<FDataLinkPin> OutputPins;
+	TemplateNode->BuildPins(InputPins, OutputPins);
+
+	// Require pin recreation if the number of 'data' pins mismatch the total number of input and output pins of the Template
+	if (DataPinCount != InputPins.Num() + OutputPins.Num())
+	{
+		return true;
+	}
+
+	auto HasMismatchingPin = [this](TConstArrayView<FDataLinkPin> InTemplatePins, EEdGraphPinDirection InDirection)
+		{
+			for (const FDataLinkPin& TemplatePin : InTemplatePins)
+			{
+				const UEdGraphPin* const FoundPin = FindPin(TemplatePin.Name, InDirection);
+
+				// Require pin recreation if the found pin that is supposed to match the template pin by name is not a data pin
+				// or has a mismatching struct
+				if (!FoundPin
+					|| FoundPin->PinType.PinCategory != UDataLinkEdGraphSchema::PC_Data
+					|| FoundPin->PinType.PinSubCategoryObject != TemplatePin.Struct)
+				{
+					return true;
+				}
+			}
+			return false;
+		};
+
+	return HasMismatchingPin(InputPins, EGPD_Input) || HasMismatchingPin(OutputPins, EGPD_Output);
+}
+
+void UDataLinkEdNode::AutowireNewNode(UEdGraphPin* InFromPin)
+{
+	if (!InFromPin)
+	{
+		return;
+	}
+
+	const UEdGraphSchema* const Schema = GetSchema();
+	if (!Schema)
+	{
+		return;
+	}
+
+	UEdGraphPin* TargetPin = nullptr;
+
+	// Iterate in reverse so that the most compatible pin end up being the first pins over the last
+	for (UEdGraphPin* Pin : ReverseIterate(Pins))
+	{
+		if (Pin && Schema->ArePinsCompatible(Pin, InFromPin))
+		{
+			TargetPin = Pin;
+
+			// If the pin names match, the best possible pin has been found
+			if (Pin->PinName == InFromPin->PinName)
+			{
+				break;
+			}
+		}
+	}
+
+	if (!TargetPin)
+	{
+		return;
+	}
+
+	if (Schema->TryCreateConnection(InFromPin, TargetPin))
+	{
+		InFromPin->GetOwningNode()->NodeConnectionListChanged();
+	}
+	else if (Schema->TryCreateConnection(TargetPin, InFromPin))
+	{
+		NodeConnectionListChanged();
+	}
+}
+
+void UDataLinkEdNode::ReconstructNode()
+{
+	Super::ReconstructNode();
+	UpdateMetadata();
+	RecreatePins();
+}
+
+bool UDataLinkEdNode::CanCreateUnderSpecifiedSchema(const UEdGraphSchema* InSchema) const
+{
+	return InSchema && InSchema->IsA<UDataLinkEdGraphSchema>();
+}
+
+FText UDataLinkEdNode::GetNodeTitle(ENodeTitleType::Type InTitleType) const
+{
+	return NodeMetadata.GetDisplayName();
+}
+
+FText UDataLinkEdNode::GetTooltipText() const
+{
+	return NodeMetadata.GetTooltipText();
+}
+
+void UDataLinkEdNode::PinConnectionListChanged(UEdGraphPin* InPin)
+{
+	Super::PinConnectionListChanged(InPin);
+	NotifyNodeChanged();
+}
+
+void UDataLinkEdNode::PostLoad()
+{
+	Super::PostLoad();
+	UpdateMetadata();
+}
+
+void UDataLinkEdNode::PostEditUndo()
+{
+	Super::PostEditUndo();
+	UpdateMetadata();
+	RecreatePins();
+}
+
+void UDataLinkEdNode::PostEditChangeChainProperty(FPropertyChangedChainEvent& InPropertyChangedEvent)
+{
+	Super::PostEditChangeChainProperty(InPropertyChangedEvent);
+
+	// Iterate over each property to see if there's a property that invalidates the node and requires node refresh
+	for (FProperty* Property : InPropertyChangedEvent.PropertyChain)
+	{
+		if (Property && Property->HasMetaData(UDataLinkEdNode::MD_InvalidatesNode))
+		{
+			UpdateMetadata();
+			RecreatePins();
+			break;
+		}
+	}
+
+	NotifyNodeChanged();
+}
+
+void UDataLinkEdNode::UpdateMetadata()
+{
+	NodeMetadata = FDataLinkNodeMetadata();
+
+	if (TemplateNode)
+	{
+		TemplateNode->BuildMetadata(NodeMetadata);
+	}
+	else
+	{
+		NodeMetadata
+			.SetDisplayName(GetClass()->GetDisplayNameText())
+			.SetTooltipText(GetClass()->GetToolTipText());
+	}
+}
+
+void UDataLinkEdNode::NotifyNodeChanged()
+{
+	if (UEdGraph* Graph = GetGraph())
+	{
+		Graph->NotifyNodeChanged(this);
+	}
+}
+
+void UDataLinkEdNode::RecreatePins()
+{
+	Modify(/*bAlwaysMarkDirty*/false);
+
+	const TArray<UEdGraphPin*> PinsCopy = Pins;
+
+	TArray<UEdGraphPin*> RemovedPins;
+	RemovedPins.Reserve(PinsCopy.Num());
+
+	// Remove all pins related to input data
+	for (UEdGraphPin* Pin : PinsCopy)
+	{
+		if (Pin)
+		{
+			Pins.Remove(Pin);
+			RemovedPins.Add(Pin);
+		}
+	}
+
+	AllocateDefaultPins();
+
+	// Build Data Link Node Pins and create them on the Ed Node
+	if (TemplateNode)
+	{
+		TArray<FDataLinkPin> InputPins;
+		TArray<FDataLinkPin> OutputPins;
+
+		TemplateNode->BuildPins(InputPins, OutputPins);
+
+		CreatePins(EGPD_Input, InputPins);
+		CreatePins(EGPD_Output, OutputPins);
+	}
+
+	// Rewire the pins to remove and matching new pins
+	if (const UEdGraphSchema* Schema = GetSchema())
+	{
+		for (UEdGraphPin* RemovedPin : RemovedPins)
+		{
+			RemovedPin->Modify(/*bAlwaysMarkDirty*/false);
+
+			if (UEdGraphPin* NewPin = FindPin(RemovedPin->PinName, RemovedPin->Direction))
+			{
+				Schema->MovePinLinks(*RemovedPin, *NewPin);
+			}
+
+			RemovedPin->MarkAsGarbage();
+			OnPinRemoved(RemovedPin);
+		}
+	}
+
+	// Refresh the UI for the graph so the pin changes show up
+	NotifyNodeChanged();
+}
+
+void UDataLinkEdNode::CreatePins(EEdGraphPinDirection InPinDirection, TConstArrayView<FDataLinkPin> InDataLinkPins)
+{
+	for (const FDataLinkPin& DataLinkPin : InDataLinkPins)
+	{
+		// No need to add if Pin already is there
+		if (!FindPin(DataLinkPin.Name, InPinDirection))
+		{
+			UEdGraphPin* const Pin = CreatePin(InPinDirection
+				, UDataLinkEdGraphSchema::PC_Data
+				, ConstCast(DataLinkPin.Struct)
+				, DataLinkPin.Name);
+
+			check(Pin);
+			Pin->PinFriendlyName = DataLinkPin.GetDisplayName();
+		}
+	}
+}
diff --git a/Engine/Plugins/Experimental/DataLink/Source/DataLinkEdGraph/Private/Nodes/DataLinkEdOutputNode.cpp b/Engine/Plugins/Experimental/DataLink/Source/DataLinkEdGraph/Private/Nodes/DataLinkEdOutputNode.cpp
new file mode 100644
--- /dev/null
+++ b/Engine/Plugins/Experimental/DataLink/Source/DataLinkEdGraph/Private/Nodes/DataLinkEdOutputNode.cpp
@@ -0,0 +1,16 @@
+// Copyright Epic Games, Inc. All Rights Reserved.
+
+#include "Nodes/DataLinkEdOutputNode.h"
+#include "DataLinkEdGraphSchema.h"
+
+FLinearColor UDataLinkEdOutputNode::GetNodeTitleColor() const
+{
+	return FLinearColor::Red;
+}
+
+void UDataLinkEdOutputNode::AllocateDefaultPins()
+{
+	// The finishing output node only has input pin.
+	// Named as "Output" to indicate result
+	CreatePin(EGPD_Input, UDataLinkEdGraphSchema::PC_Data, UDataLinkEdNode::PN_Output);
+}
diff --git a/Engine/Plugins/Experimental/DataLink/Source/DataLinkEdGraph/Public/Blueprint/K2Node_DataLinkRequest.h b/Engine/Plugins/Experimental/DataLink/Source/DataLinkEdGraph/Public/Blueprint/K2Node_DataLinkRequest.h
new file mode 100644
--- /dev/null
+++ b/Engine/Plugins/Experimental/DataLink/Source/DataLinkEdGraph/Public/Blueprint/K2Node_DataLinkRequest.h
@@ -0,0 +1,39 @@
+// Copyright Epic Games, Inc. All Rights Reserved.
+
+#pragma once
+
+#include "K2Node.h"
+#include "K2Node_DataLinkRequest.generated.h"
+
+class UDataLinkGraph;
+
+UCLASS(MinimalAPI)
+class UK2Node_DataLinkRequest : public UK2Node
+{
+	GENERATED_BODY()
+
+public:
+	static const FLazyName PN_DataLinkInstance;
+	static const FLazyName PN_ExecutionContext;
+	static const FLazyName PN_DataLinkSinkProvider;
+	static const FLazyName PN_Success;
+	static const FLazyName PN_Failure;
+
+	UEdGraphPin* GetDataLinkInstancePin() const;
+	UEdGraphPin* GetExecutionContextPin() const;
+	UEdGraphPin* GetDataLinkSinkProviderPin() const;
+	UEdGraphPin* GetOutputPin() const;
+
+	//~ Begin UEdGraphNode
+	virtual void AllocateDefaultPins() override;
+	virtual bool IsCompatibleWithGraph(const UEdGraph* InTargetGraph) const override;
+	virtual FText GetNodeTitle(ENodeTitleType::Type InTitleType) const override;
+	//~ End UEdGraphNode
+
+	//~ Begin UK2Node
+	virtual void ExpandNode(FKismetCompilerContext& InCompilerContext, UEdGraph* InSourceGraph) override;
+	virtual FName GetCornerIcon() const override;
+	virtual void GetMenuActions(FBlueprintActionDatabaseRegistrar& InActionRegistrar) const override;
+	virtual FText GetMenuCategory() const override;
+	//~ End UK2Node
+};
diff --git a/Engine/Plugins/Experimental/DataLink/Source/DataLinkEdGraph/Public/DataLinkEdGraph.h b/Engine/Plugins/Experimental/DataLink/Source/DataLinkEdGraph/Public/DataLinkEdGraph.h
new file mode 100644
--- /dev/null
+++ b/Engine/Plugins/Experimental/DataLink/Source/DataLinkEdGraph/Public/DataLinkEdGraph.h
@@ -0,0 +1,51 @@
+// Copyright Epic Games, Inc. All Rights Reserved.
+
+#pragma once
+
+#include "EdGraph/EdGraph.h"
+#include "DataLinkEdGraph.generated.h"
+
+class UDataLinkEdOutputNode;
+class UDataLinkGraph;
+
+UCLASS(MinimalAPI)
+class UDataLinkEdGraph : public UEdGraph
+{
+    GENERATED_BODY()
+
+public:
+	UDataLinkEdGraph();
+
+    /**
+     * Finds the 'cosmetic' output node
+     * @see UDataLinkEdOutputNode
+     */
+    DATALINKEDGRAPH_API UDataLinkEdOutputNode* FindOutputNode() const;
+
+	/** Initializes all the nodes in this graph, recreating pins of outdated nodes to ensure these are up-to-date with their templates */
+	DATALINKEDGRAPH_API void InitializeNodes();
+
+	//~ Begin UObject
+	DATALINKEDGRAPH_API virtual void BeginDestroy() override;
+	//~ End UObject
+
+	DATALINKEDGRAPH_API void DirtyGraph();
+
+	DATALINKEDGRAPH_API bool IsCompiledGraphUpToDate() const;
+
+private:
+	void OnGraphCompiled(UDataLinkGraph* InCompiledGraph);
+
+	/**
+	 * Represents id of graph changes that haven't been compiled yet.
+	 * Note: if a change occurs, this id will be regenerated only if the Last Compiled Change Id matches the current change id.
+	 */
+	UPROPERTY()
+	FGuid ChangeId;
+
+	/** The Change Id that was last compiled */
+	UPROPERTY()
+	FGuid LastCompiledChangeId;
+
+	FDelegateHandle OnGraphCompiledHandle;
+};
diff --git a/Engine/Plugins/Experimental/DataLink/Source/DataLinkEdGraph/Public/DataLinkEdGraphSchema.h b/Engine/Plugins/Experimental/DataLink/Source/DataLinkEdGraph/Public/DataLinkEdGraphSchema.h
new file mode 100644
--- /dev/null
+++ b/Engine/Plugins/Experimental/DataLink/Source/DataLinkEdGraph/Public/DataLinkEdGraphSchema.h
@@ -0,0 +1,32 @@
+// Copyright Epic Games, Inc. All Rights Reserved.
+
+#pragma once
+
+#include "EdGraph/EdGraphSchema.h"
+#include "DataLinkEdGraphSchema.generated.h"
+
+UCLASS(MinimalAPI)
+class UDataLinkEdGraphSchema : public UEdGraphSchema
+{
+	GENERATED_BODY()
+
+public:
+	/** Pin Categories */
+	DATALINKEDGRAPH_API static const FLazyName PC_Data;
+
+	/** Pin Category Colors */
+	DATALINKEDGRAPH_API static const FLinearColor PCC_Data;
+
+	/** Check whether connecting these pins would cause a loop */
+	bool IsConnectionLooping(const UEdGraphPin* InInputPin, const UEdGraphPin* InOutputPin) const;
+
+	//~ Begin UEdGraphSchema
+	virtual void CreateDefaultNodesForGraph(UEdGraph& InGraph) const override;
+	virtual bool ArePinsCompatible(const UEdGraphPin* InPinA, const UEdGraphPin* InPinB, const UClass* InCallingContext = nullptr, bool bInIgnoreArray = false) const override;
+	virtual void GetGraphContextActions(FGraphContextMenuBuilder& InContextMenuBuilder) const override;
+	virtual void GetContextMenuActions(UToolMenu* InMenu, UGraphNodeContextMenuContext* InContext) const override;
+	virtual const FPinConnectionResponse CanCreateConnection(const UEdGraphPin* InSourcePin, const UEdGraphPin* InTargetPin) const override;
+	virtual FLinearColor GetPinTypeColor(const FEdGraphPinType& InPinType) const override;
+	virtual void GetGraphDisplayInformation(const UEdGraph& InGraph, FGraphDisplayInfo& OutDisplayInfo) const override;
+	//~ End UEdGraphSchema
+};
diff --git a/Engine/Plugins/Experimental/DataLink/Source/DataLinkEdGraph/Public/Nodes/DataLinkEdNode.h b/Engine/Plugins/Experimental/DataLink/Source/DataLinkEdGraph/Public/Nodes/DataLinkEdNode.h
new file mode 100644
--- /dev/null
+++ b/Engine/Plugins/Experimental/DataLink/Source/DataLinkEdGraph/Public/Nodes/DataLinkEdNode.h
@@ -0,0 +1,73 @@
+// Copyright Epic Games, Inc. All Rights Reserved.
+
+#pragma once
+
+#include "DataLinkNodeMetadata.h"
+#include "EdGraph/EdGraphNode.h"
+#include "DataLinkEdNode.generated.h"
+
+class UDataLinkNode;
+struct FDataLinkPin;
+template <typename T> class TSubclassOf;
+
+UCLASS(MinimalAPI)
+class UDataLinkEdNode : public UEdGraphNode
+{
+	GENERATED_BODY()
+
+public:
+	DATALINKEDGRAPH_API static const FLazyName MD_InvalidatesNode;
+	DATALINKEDGRAPH_API static const FLazyName PN_Output;
+
+	DATALINKEDGRAPH_API void SetTemplateNodeClass(TSubclassOf<UDataLinkNode> InNodeClass, bool bInReconstructNode = true);
+
+	DATALINKEDGRAPH_API void ForEachPinConnection(TFunctionRef<void(const UEdGraphPin&, const UDataLinkEdNode&, const UEdGraphPin&)> InFunction) const;
+
+	UDataLinkNode* GetTemplateNode() const
+	{
+		return TemplateNode;
+	}
+
+	//~ Begin UEdGraphNode
+	DATALINKEDGRAPH_API virtual void AutowireNewNode(UEdGraphPin* InFromPin) override;
+	DATALINKEDGRAPH_API virtual void ReconstructNode() override;
+	DATALINKEDGRAPH_API virtual bool CanCreateUnderSpecifiedSchema(const UEdGraphSchema* InSchema) const override;
+	DATALINKEDGRAPH_API virtual FText GetNodeTitle(ENodeTitleType::Type InTitleType) const override;
+	DATALINKEDGRAPH_API virtual FText GetTooltipText() const override;
+	DATALINKEDGRAPH_API virtual void PinConnectionListChanged(UEdGraphPin* InPin) override;
+	//~ End UEdGraphNode
+
+	//~ Begin UObject
+	DATALINKEDGRAPH_API virtual void PostLoad() override;
+	DATALINKEDGRAPH_API virtual void PostEditUndo() override;
+	DATALINKEDGRAPH_API virtual void PostEditChangeChainProperty(FPropertyChangedChainEvent& InPropertyChangedEvent) override;
+	//~ End UObject
+
+	static FName GetTemplateNodePropertyName()
+	{
+		return GET_MEMBER_NAME_CHECKED(UDataLinkEdNode, TemplateNode);
+	}
+
+	void UpdateMetadata();
+
+	DATALINKEDGRAPH_API void NotifyNodeChanged();
+
+	/**
+ 	 * Checks against the current Template Node to determine if the Ed Node's structure matches the Template Node's
+ 	 * @return true if the structures mismatch and requires node reconstruction
+ 	 */
+	DATALINKEDGRAPH_API virtual bool RequiresPinRecreation() const;
+
+	/** Destroys existing pins and recreates new ones based on the underlying Template Node while also restoring pin links where available */
+	void RecreatePins();
+
+private:
+	void CreatePins(EEdGraphPinDirection InPinDirection, TConstArrayView<FDataLinkPin> InDataLinkPins);
+
+	/** Underlying Data Link Node this Editor Node represents */
+	UPROPERTY(VisibleAnywhere, Instanced, Category="Data Link")
+	TObjectPtr<UDataLinkNode> TemplateNode;
+
+	/** Cached metadata from the Template Node */
+	FDataLinkNodeMetadata NodeMetadata;
+};
diff --git a/Engine/Plugins/Experimental/DataLink/Source/DataLinkEdGraph/Public/Nodes/DataLinkEdOutputNode.h b/Engine/Plugins/Experimental/DataLink/Source/DataLinkEdGraph/Public/Nodes/DataLinkEdOutputNode.h
new file mode 100644
--- /dev/null
+++ b/Engine/Plugins/Experimental/DataLink/Source/DataLinkEdGraph/Public/Nodes/DataLinkEdOutputNode.h
@@ -0,0 +1,30 @@
+// Copyright Epic Games, Inc. All Rights Reserved.
+
+#pragma once
+
+#include "DataLinkEdNode.h"
+#include "DataLinkEdOutputNode.generated.h"
+
+/**
+ * A 'cosmetic' output node.
+ * This is considered cosmetic because it does not get compiled in and
+ * is used only to iterate through the nodes that should be considered for compilation.
+ */
+UCLASS(MinimalAPI, DisplayName="Output")
+class UDataLinkEdOutputNode : public UDataLinkEdNode
+{
+	GENERATED_BODY()
+
+public:
+	//~ Begin UDataLinkEdNode
+	virtual bool RequiresPinRecreation() const override { return false; }
+	//~ End UDataLinkEdNode
+
+	//~ Begin UEdGraphNode
+	DATALINKEDGRAPH_API virtual FLinearColor GetNodeTitleColor() const;
+	DATALINKEDGRAPH_API virtual void AllocateDefaultPins() override;
+	virtual bool CanDuplicateNode() const override { return false; }
+	virtual bool CanUserDeleteNode() const override { return false; }
+	virtual bool GetCanRenameNode() const override { return false; }
+	//~ End UEdGraphNode
+};
diff --git a/Engine/Plugins/Experimental/DataLink/Source/DataLinkEditor/DataLinkEditor.Build.cs b/Engine/Plugins/Experimental/DataLink/Source/DataLinkEditor/DataLinkEditor.Build.cs
new file mode 100644
--- /dev/null
+++ b/Engine/Plugins/Experimental/DataLink/Source/DataLinkEditor/DataLinkEditor.Build.cs
@@ -0,0 +1,37 @@
+// Copyright Epic Games, Inc. All Rights Reserved.
+
+using UnrealBuildTool;
+
+public class DataLinkEditor : ModuleRules
+{
+    public DataLinkEditor(ReadOnlyTargetRules Target) : base(Target)
+    {
+        PCHUsage = ModuleRules.PCHUsageMode.UseExplicitOrSharedPCHs;
+
+        PublicDependencyModuleNames.AddRange(
+            new string[]
+            {
+                "Core",
+                "CoreUObject",
+            }
+        );
+
+        PrivateDependencyModuleNames.AddRange(
+            new string[]
+            {
+                "ApplicationCore",
+                "AssetDefinition",
+                "DataLink",
+                "DataLinkEdGraph",
+                "Engine",
+                "GraphEditor",
+                "PropertyEditor",
+                "Slate",
+                "SlateCore",
+                "StructUtilsEditor",
+                "ToolMenus",
+                "UnrealEd",
+            }
+        );
+    }
+}
diff --git a/Engine/Plugins/Experimental/DataLink/Source/DataLinkEditor/Private/AssetDefinition_DataLinkGraph.cpp b/Engine/Plugins/Experimental/DataLink/Source/DataLinkEditor/Private/AssetDefinition_DataLinkGraph.cpp
new file mode 100644
--- /dev/null
+++ b/Engine/Plugins/Experimental/DataLink/Source/DataLinkEditor/Private/AssetDefinition_DataLinkGraph.cpp
@@ -0,0 +1,44 @@
+// Copyright Epic Games, Inc. All Rights Reserved.
+
+#include "AssetDefinition_DataLinkGraph.h"
+#include "DataLinkGraph.h"
+#include "DataLinkGraphAssetEditor.h"
+
+#define LOCTEXT_NAMESPACE "AssetDefinition_DataLinkGraph"
+
+FText UAssetDefinition_DataLinkGraph::GetAssetDisplayName() const
+{
+	return LOCTEXT("AssetDisplayName", "Data Link Graph");
+}
+
+FLinearColor UAssetDefinition_DataLinkGraph::GetAssetColor() const
+{
+	return FLinearColor(FColor(64, 130, 109));
+}
+
+TSoftClassPtr<UObject> UAssetDefinition_DataLinkGraph::GetAssetClass() const
+{
+	return UDataLinkGraph::StaticClass();
+}
+
+TConstArrayView<FAssetCategoryPath> UAssetDefinition_DataLinkGraph::GetAssetCategories() const
+{
+	return MakeArrayView(&EAssetCategoryPaths::Misc, 1);
+}
+
+EAssetCommandResult UAssetDefinition_DataLinkGraph::OpenAssets(const FAssetOpenArgs& InOpenArgs) const
+{
+	check(GEditor);
+	UAssetEditorSubsystem* AssetEditorSubsystem = GEditor->GetEditorSubsystem<UAssetEditorSubsystem>();
+	check(AssetEditorSubsystem);
+
+	for (UDataLinkGraph* DataLinkGraph : InOpenArgs.LoadObjects<UDataLinkGraph>())
+	{
+		UDataLinkGraphAssetEditor* AssetEditor = NewObject<UDataLinkGraphAssetEditor>(AssetEditorSubsystem, NAME_None, RF_Transient);
+		AssetEditor->Initialize(DataLinkGraph);
+	}
+
+	return EAssetCommandResult::Handled;
+}
+
+#undef LOCTEXT_NAMESPACE
diff --git a/Engine/Plugins/Experimental/DataLink/Source/DataLinkEditor/Private/AssetDefinition_DataLinkGraph.h b/Engine/Plugins/Experimental/DataLink/Source/DataLinkEditor/Private/AssetDefinition_DataLinkGraph.h
new file mode 100644
--- /dev/null
+++ b/Engine/Plugins/Experimental/DataLink/Source/DataLinkEditor/Private/AssetDefinition_DataLinkGraph.h
@@ -0,0 +1,21 @@
+// Copyright Epic Games, Inc. All Rights Reserved.
+
+#pragma once
+
+#include "AssetDefinition.h"
+#include "AssetDefinition_DataLinkGraph.generated.h"
+
+UCLASS()
+class UAssetDefinition_DataLinkGraph : public UAssetDefinition
+{
+	GENERATED_BODY()
+
+public:
+	//~ Begin UAssetDefinition
+	virtual FText GetAssetDisplayName() const override;
+	virtual FLinearColor GetAssetColor() const override;
+	virtual TSoftClassPtr<UObject> GetAssetClass() const override;
+	virtual TConstArrayView<FAssetCategoryPath> GetAssetCategories() const override;
+	virtual EAssetCommandResult OpenAssets(const FAssetOpenArgs& InOpenArgs) const override;
+	//~ End UAssetDefinition
+};
diff --git a/Engine/Plugins/Experimental/DataLink/Source/DataLinkEditor/Private/Compiler/DataLinkGraphCompileStatus.h b/Engine/Plugins/Experimental/DataLink/Source/DataLinkEditor/Private/Compiler/DataLinkGraphCompileStatus.h
new file mode 100644
--- /dev/null
+++ b/Engine/Plugins/Experimental/DataLink/Source/DataLinkEditor/Private/Compiler/DataLinkGraphCompileStatus.h
@@ -0,0 +1,13 @@
+// Copyright Epic Games, Inc. All Rights Reserved.
+
+#pragma once
+
+#include "HAL/Platform.h"
+
+enum class EDataLinkGraphCompileStatus : uint8
+{
+	Unknown,
+	Warning,
+	Error,
+	Good,
+};
diff --git a/Engine/Plugins/Experimental/DataLink/Source/DataLinkEditor/Private/Compiler/DataLinkGraphCompiler.cpp b/Engine/Plugins/Experimental/DataLink/Source/DataLinkEditor/Private/Compiler/DataLinkGraphCompiler.cpp
new file mode 100644
--- /dev/null
+++ b/Engine/Plugins/Experimental/DataLink/Source/DataLinkEditor/Private/Compiler/DataLinkGraphCompiler.cpp
@@ -0,0 +1,286 @@
+// Copyright Epic Games, Inc. All Rights Reserved.
+
+#include "DataLinkGraphCompiler.h"
+#include "DataLinkEditorLog.h"
+#include "DataLinkGraph.h"
+#include "DataLinkGraphCompileStatus.h"
+#include "DataLinkNode.h"
+#include "EdGraph/EdGraphPin.h"
+#include "Nodes/DataLinkEdNode.h"
+#include "Nodes/DataLinkEdOutputNode.h"
+
+FDataLinkGraphCompiler::FDataLinkGraphCompiler(UDataLinkGraph* InDataLinkGraph)
+	: DataLinkGraph(InDataLinkGraph)
+{
+}
+
+EDataLinkGraphCompileStatus FDataLinkGraphCompiler::Compile()
+{
+	if (!DataLinkGraph)
+	{
+		UE_LOG(LogDataLinkEditor, Error, TEXT("Compilation failed. Data Link Graph is invalid."));
+		return EDataLinkGraphCompileStatus::Error;
+	}
+
+	DataLinkEdGraph = Cast<UDataLinkEdGraph>(DataLinkGraph->GetEdGraph());
+	if (!DataLinkEdGraph)
+	{
+		UE_LOG(LogDataLinkEditor, Error, TEXT("Compilation failed. Data Link Ed Graph is invalid in graph '%s'"), *DataLinkGraph->GetName());
+		return EDataLinkGraphCompileStatus::Error;
+	}
+
+	CleanExistingGraph();
+
+	if (!CompileNodes())
+	{
+		// Compilation failed. Clean graph again from what had started compiling
+    	CleanExistingGraph();
+    	return EDataLinkGraphCompileStatus::Error;
+	}
+
+	UDataLinkGraph::OnGraphCompiledDelegate.Broadcast(DataLinkGraph);
+	return EDataLinkGraphCompileStatus::Good;
+}
+
+void FDataLinkGraphCompiler::CleanExistingGraph()
+{
+	DataLinkGraph->InputNodes.Reset();
+	DataLinkGraph->OutputNode = nullptr;
+	DataLinkGraph->Nodes.Reset();
+}
+
+bool FDataLinkGraphCompiler::CompileNodes()
+{
+	if (DataLinkEdGraph->Nodes.IsEmpty())
+	{
+		// Nothing further to compile. Unknown Data Link Graph
+		return false;
+	}
+
+	// Step 1: Compile the nodes and fill in the Ed to Compiled Node map
+	if (!CreateCompiledNodes())
+	{
+		return false;
+	}
+
+	// Step 2: Set the node links by finding the Editor link and using the map to find the compiled node
+	LinkNodes();
+
+	// Step 3: Find and set the inputs and output nodes
+	SetInputOutputNodes();
+
+	return true;
+}
+
+bool FDataLinkGraphCompiler::CreateCompiledNodes()
+{
+	UDataLinkEdOutputNode* OutputNode = DataLinkEdGraph->FindOutputNode();
+	if (!OutputNode)
+	{
+		UE_LOG(LogDataLinkEditor, Error, TEXT("Compilation failed. Output Node was not valid! Graph: %s")
+			, *DataLinkGraph->GetName());
+		return false;
+	}
+
+	EdToCompiledMap.Empty(DataLinkEdGraph->Nodes.Num());
+
+	TArray<const UDataLinkEdNode*> EdNodesRemaining;
+
+	const auto AddLinkedNodes =
+		[&EdToCompiledMap = EdToCompiledMap, &EdNodesRemaining](const UEdGraphPin& InPin, const UDataLinkEdNode& InLinkedNode, const UEdGraphPin& InLinkedPin)
+		{
+			// Skip pin connections that aren't input
+			// Skip nodes that have already been compiled
+			if (InPin.Direction == EGPD_Input && !EdToCompiledMap.Contains(&InLinkedNode))
+			{
+				// Add the node only once (in cases where multiple pins are connected to one node)
+				EdNodesRemaining.AddUnique(&InLinkedNode);
+			}
+		};
+
+	OutputNode->ForEachPinConnection(AddLinkedNodes);
+
+	if (EdNodesRemaining.IsEmpty())
+	{
+		UE_LOG(LogDataLinkEditor, Error, TEXT("Compilation failed. No nodes were considered for compilation in graph '%s'"), *DataLinkGraph->GetName());
+		return false;
+	}
+
+	// The node that the cosmetic output node is connected to is the actual output node at runtime
+	OutputEdNode = EdNodesRemaining[0];
+
+	while (!EdNodesRemaining.IsEmpty())
+	{
+		const UDataLinkEdNode* EdNode = EdNodesRemaining.Pop();
+		if (!EdNode)
+		{
+			continue;
+		}
+
+		if (UDataLinkNode* TemplateNode = EdNode->GetTemplateNode())
+		{
+			EdToCompiledMap.Add(EdNode, CompileNode(TemplateNode));
+		}
+		else
+		{
+			UE_LOG(LogDataLinkEditor, Error, TEXT("Compilation failed. EdNode '%s' did not have a valid Template Node! Graph: %s"), *EdNode->GetName(), *DataLinkGraph->GetName());
+			return false;
+		}
+
+		// Add connected nodes in the input direction
+		EdNode->ForEachPinConnection(AddLinkedNodes);
+	}
+
+	if (EdToCompiledMap.IsEmpty())
+	{
+		UE_LOG(LogDataLinkEditor, Error, TEXT("Compilation failed. No nodes were considered for compilation in graph '%s'"), *DataLinkGraph->GetName());
+		return false;	
+	}
+
+	// Populate the Nodes array
+	DataLinkGraph->Nodes.Reserve(EdToCompiledMap.Num());
+	for (const TPair<const UDataLinkEdNode*, UDataLinkNode*>& Pair : EdToCompiledMap)
+	{
+		DataLinkGraph->Nodes.Add(Pair.Value);
+	}
+
+	return true;
+}
+
+UDataLinkNode* FDataLinkGraphCompiler::CompileNode(UDataLinkNode* InTemplateNode)
+{
+	check(InTemplateNode);
+
+	UDataLinkNode* CompiledNode = NewObject<UDataLinkNode>(DataLinkGraph
+		, InTemplateNode->GetClass()
+		, NAME_None
+		, RF_NoFlags
+		, InTemplateNode);
+
+	// Build Pins
+	{
+		TArray<FDataLinkPin> InputPins;
+		TArray<FDataLinkPin> OutputPins;
+
+		CompiledNode->BuildPins(InputPins, OutputPins);
+
+		CompiledNode->InputPins = MoveTemp(InputPins);
+		CompiledNode->OutputPins = MoveTemp(OutputPins);
+	}
+
+	return CompiledNode;
+}
+
+void FDataLinkGraphCompiler::LinkNodes()
+{
+	for (UEdGraphNode* EdNodeRaw : DataLinkEdGraph->Nodes)
+	{
+		UDataLinkEdNode* EdNode = Cast<UDataLinkEdNode>(EdNodeRaw);
+		if (!EdNode)
+		{
+			continue;
+		}
+
+		UDataLinkNode* CompiledNode = FindCompiledNode(EdNode);
+		if (!CompiledNode)
+		{
+			continue;
+		}
+
+		// Update Inputs and Output Pins to have the compiled connected node
+		EdNode->ForEachPinConnection(
+			[this, CompiledNode](const UEdGraphPin& InPin, const UDataLinkEdNode& InLinkedNode, const UEdGraphPin& InLinkedPin)
+			{
+				const bool bIsInputPin = InPin.Direction == EGPD_Input;
+
+				FDataLinkPin* const Pin = bIsInputPin
+					? CompiledNode->InputPins.FindByKey(InPin.PinName)
+					: CompiledNode->OutputPins.FindByKey(InPin.PinName);
+
+				if (!Pin)
+				{
+					return;
+				}
+
+				if (UDataLinkNode* LinkedNode = FindCompiledNode(&InLinkedNode))
+				{
+					Pin->LinkedNode = LinkedNode;
+
+					// Link Index is the opposite since Input Pins are connected to Output Pins
+					Pin->LinkedIndex = !bIsInputPin
+						? LinkedNode->InputPins.IndexOfByKey(InLinkedPin.PinName)
+						: LinkedNode->OutputPins.IndexOfByKey(InLinkedPin.PinName);
+				}
+				else
+				{
+					Pin->LinkedNode = nullptr;
+					Pin->LinkedIndex = INDEX_NONE;
+				}
+			});
+	}
+}
+
+bool FDataLinkGraphCompiler::SetInputOutputNodes()
+{
+	// at this point in compilation the map shouldn't be empty
+	check(!EdToCompiledMap.IsEmpty());
+
+	UDataLinkNode* OutputNode = FindCompiledNode(OutputEdNode);
+	if (!ensureAlways(OutputNode))
+	{
+		UE_LOG(LogDataLinkEditor, Error, TEXT("Compilation failed. Output Node was not found in graph '%s'")
+			, *DataLinkGraph->GetName());
+		return false;
+	}
+
+	DataLinkGraph->OutputNode = OutputNode;
+
+	// Recursively adds all the graph's entry nodes (i.e. nodes with input pins that are not connected to other nodes)
+	AddGraphEntryNodes(OutputNode);
+
+	if (DataLinkGraph->InputNodes.IsEmpty())
+	{
+		UE_LOG(LogDataLinkEditor, Error, TEXT("Compilation failed. Input nodes could not be determined in graph '%s'"), *DataLinkGraph->GetName())
+		return false;
+	}
+
+	return true;
+}
+
+void FDataLinkGraphCompiler::AddGraphEntryNodes(const UDataLinkNode* InNode)
+{
+	TConstArrayView<FDataLinkPin> InputPins = InNode->GetInputPins();
+
+	// Node with no input pins are considered entry nodes
+	if (InputPins.IsEmpty())
+	{
+		DataLinkGraph->InputNodes.AddUnique(const_cast<UDataLinkNode*>(InNode));
+		return;
+	}
+
+	bool bAddedInputNode = false;
+
+	for (const FDataLinkPin& InputPin : InputPins)
+	{
+		// If there is a linked node to follow, this node is not the entry node for this pin path
+		if (InputPin.LinkedNode)
+		{
+			AddGraphEntryNodes(InputPin.LinkedNode);
+		}
+		else if (!bAddedInputNode)
+		{
+			// Avoid adding the same node multiple times if multiple of its input pins are 'opened' (i.e. not linked to other node)
+			DataLinkGraph->InputNodes.AddUnique(const_cast<UDataLinkNode*>(InNode));
+			bAddedInputNode = true;
+		}
+	}
+}
+
+UDataLinkNode* FDataLinkGraphCompiler::FindCompiledNode(const UDataLinkEdNode* InEdNode) const
+{
+	if (UDataLinkNode* const* FoundNode = EdToCompiledMap.Find(InEdNode))
+	{
+		return *FoundNode;
+	}
+	return nullptr;
+}
diff --git a/Engine/Plugins/Experimental/DataLink/Source/DataLinkEditor/Private/Compiler/DataLinkGraphCompiler.h b/Engine/Plugins/Experimental/DataLink/Source/DataLinkEditor/Private/Compiler/DataLinkGraphCompiler.h
new file mode 100644
--- /dev/null
+++ b/Engine/Plugins/Experimental/DataLink/Source/DataLinkEditor/Private/Compiler/DataLinkGraphCompiler.h
@@ -0,0 +1,44 @@
+// Copyright Epic Games, Inc. All Rights Reserved.
+
+#pragma once
+
+#include "DataLinkEdGraph.h"
+
+class UDataLinkEdNode;
+class UDataLinkGraph;
+class UDataLinkNode;
+enum class EDataLinkGraphCompileStatus : uint8;
+
+class FDataLinkGraphCompiler
+{
+public:
+	explicit FDataLinkGraphCompiler(UDataLinkGraph* InDataLinkGraph);
+
+	EDataLinkGraphCompileStatus Compile();
+
+private:
+	void CleanExistingGraph();
+
+	bool CompileNodes();
+
+	bool CreateCompiledNodes();
+
+	UDataLinkNode* CompileNode(UDataLinkNode* InTemplateNode);
+
+	void LinkNodes();
+
+	bool SetInputOutputNodes();
+
+	void AddGraphEntryNodes(const UDataLinkNode* InNode);
+
+	UDataLinkNode* FindCompiledNode(const UDataLinkEdNode* InEdNode) const;
+
+	TObjectPtr<UDataLinkGraph> DataLinkGraph;
+
+	TObjectPtr<UDataLinkEdGraph> DataLinkEdGraph;
+
+	TObjectPtr<const UDataLinkEdNode> OutputEdNode;
+
+	/** Map of the Editor Node to its Compiled Node */
+	TMap<const UDataLinkEdNode*, UDataLinkNode*> EdToCompiledMap;
+};
diff --git a/Engine/Plugins/Experimental/DataLink/Source/DataLinkEditor/Private/Compiler/DataLinkGraphCompilerTool.cpp b/Engine/Plugins/Experimental/DataLink/Source/DataLinkEditor/Private/Compiler/DataLinkGraphCompilerTool.cpp
new file mode 100644
--- /dev/null
+++ b/Engine/Plugins/Experimental/DataLink/Source/DataLinkEditor/Private/Compiler/DataLinkGraphCompilerTool.cpp
@@ -0,0 +1,105 @@
+// Copyright Epic Games, Inc. All Rights Reserved.
+
+#include "DataLinkGraphCompilerTool.h"
+#include "DataLinkEditorStyle.h"
+#include "DataLinkGraphAssetEditor.h"
+#include "DataLinkGraphAssetToolkit.h"
+#include "DataLinkGraphCommands.h"
+#include "DataLinkGraphCompiler.h"
+#include "DataLinkGraphEditorMenuContext.h"
+#include "ToolMenu.h"
+#include "ToolMenuSection.h"
+
+#define LOCTEXT_NAMESPACE "DataLinkGraphCompilerToolkit"
+
+namespace UE::DataLinkEditor::Private
+{
+	const TMap<EDataLinkGraphCompileStatus, FLazyName> CompileStatusBackgrounds =
+	{
+		{ EDataLinkGraphCompileStatus::Unknown , TEXT("CompileStatus.Background.Unknown") },
+		{ EDataLinkGraphCompileStatus::Warning , TEXT("CompileStatus.Background.Warning") },
+		{ EDataLinkGraphCompileStatus::Error   , TEXT("CompileStatus.Background.Error")  },
+		{ EDataLinkGraphCompileStatus::Good    , TEXT("CompileStatus.Background.Good")   },
+	};
+
+	const TMap<EDataLinkGraphCompileStatus, FLazyName> CompileStatusOverlays =
+	{
+		{ EDataLinkGraphCompileStatus::Unknown , TEXT("CompileStatus.Overlay.Unknown") },
+		{ EDataLinkGraphCompileStatus::Warning , TEXT("CompileStatus.Overlay.Warning") },
+		{ EDataLinkGraphCompileStatus::Error   , TEXT("CompileStatus.Overlay.Error")  },
+		{ EDataLinkGraphCompileStatus::Good    , TEXT("CompileStatus.Overlay.Good")   },
+	};
+}
+
+void FDataLinkGraphCompilerTool::ExtendMenu(UToolMenu* InMenu)
+{
+	FToolMenuSection& CompilerSection = InMenu->FindOrAddSection(TEXT("Compile")
+		, TAttribute<FText>()
+		, FToolMenuInsert(TEXT("Asset"), EToolMenuInsertType::After));
+
+	CompilerSection.AddDynamicEntry(TEXT("Compiler")
+		, FNewToolMenuSectionDelegate::CreateStatic(&FDataLinkGraphCompilerTool::ExtendDynamicCompilerSection));
+}
+
+FDataLinkGraphCompilerTool::FDataLinkGraphCompilerTool(UDataLinkGraphAssetEditor* InAssetEditor)
+	: AssetEditor(InAssetEditor)
+	, LastCompiledStatus(EDataLinkGraphCompileStatus::Unknown)
+{
+	UDataLinkEdGraph* EdGraph = AssetEditor->GetDataLinkEdGraph();
+	if (EdGraph && EdGraph->IsCompiledGraphUpToDate())
+	{
+		LastCompiledStatus = EDataLinkGraphCompileStatus::Good;
+	}
+}
+
+void FDataLinkGraphCompilerTool::BindCommands(const TSharedRef<FUICommandList>& InCommandList)
+{
+	const FDataLinkGraphCommands& GraphCommands = FDataLinkGraphCommands::Get();
+	InCommandList->MapAction(GraphCommands.Compile, FExecuteAction::CreateSP(this, &FDataLinkGraphCompilerTool::Compile));
+}
+
+void FDataLinkGraphCompilerTool::Compile()
+{
+	FDataLinkGraphCompiler Compiler(AssetEditor->GetDataLinkGraph());
+	LastCompiledStatus = Compiler.Compile();
+}
+
+void FDataLinkGraphCompilerTool::ExtendDynamicCompilerSection(FToolMenuSection& InSection)
+{
+	UDataLinkGraphEditorMenuContext* MenuContext = InSection.FindContext<UDataLinkGraphEditorMenuContext>();
+	if (!MenuContext)
+	{
+		return;
+	}
+
+	TSharedPtr<FDataLinkGraphAssetToolkit> Toolkit = MenuContext->ToolkitWeak.Pin();
+	if (!Toolkit.IsValid())
+	{
+		return;
+	}
+
+	const FDataLinkGraphCommands& GraphCommands = FDataLinkGraphCommands::Get();
+
+	InSection.AddEntry(FToolMenuEntry::InitToolBarButton(GraphCommands.Compile
+		, TAttribute<FText>()
+		, TAttribute<FText>()
+		, TAttribute<FSlateIcon>::CreateSP(&Toolkit->GetCompilerTool(), &FDataLinkGraphCompilerTool::GetCompileIcon)));
+}
+
+FSlateIcon FDataLinkGraphCompilerTool::GetCompileIcon() const
+{
+	EDataLinkGraphCompileStatus CompileStatus = LastCompiledStatus;
+
+	UDataLinkEdGraph* EdGraph = AssetEditor->GetDataLinkEdGraph();
+	if (EdGraph && !EdGraph->IsCompiledGraphUpToDate())
+	{
+		CompileStatus = EDataLinkGraphCompileStatus::Unknown;
+	}
+
+	const FName Background = UE::DataLinkEditor::Private::CompileStatusBackgrounds[CompileStatus]; 
+	const FName Overlay = UE::DataLinkEditor::Private::CompileStatusOverlays[CompileStatus]; 
+
+	return FSlateIcon(FDataLinkEditorStyle::Get().GetStyleSetName(), Background, Background, Overlay);
+}
+
+#undef LOCTEXT_NAMESPACE
diff --git a/Engine/Plugins/Experimental/DataLink/Source/DataLinkEditor/Private/Compiler/DataLinkGraphCompilerTool.h b/Engine/Plugins/Experimental/DataLink/Source/DataLinkEditor/Private/Compiler/DataLinkGraphCompilerTool.h
new file mode 100644
--- /dev/null
+++ b/Engine/Plugins/Experimental/DataLink/Source/DataLinkEditor/Private/Compiler/DataLinkGraphCompilerTool.h
@@ -0,0 +1,35 @@
+// Copyright Epic Games, Inc. All Rights Reserved.
+
+#pragma once
+
+#include "DataLinkGraphCompileStatus.h"
+#include "Templates/SharedPointer.h"
+#include "UObject/ObjectPtr.h"
+
+class FUICommandList;
+class UDataLinkGraphAssetEditor;
+class UToolMenu;
+struct FSlateIcon;
+struct FToolMenuSection;
+
+class FDataLinkGraphCompilerTool : public TSharedFromThis<FDataLinkGraphCompilerTool>
+{
+public:
+	static void ExtendMenu(UToolMenu* InMenu);
+
+	explicit FDataLinkGraphCompilerTool(UDataLinkGraphAssetEditor* InAssetEditor);
+
+	void BindCommands(const TSharedRef<FUICommandList>& InCommandList);
+
+	void Compile();
+
+private:
+	static void ExtendDynamicCompilerSection(FToolMenuSection& InSection);
+
+	FSlateIcon GetCompileIcon() const;
+
+	TObjectPtr<UDataLinkGraphAssetEditor> AssetEditor;
+
+	EDataLinkGraphCompileStatus LastCompiledStatus;
+};
+
diff --git a/Engine/Plugins/Experimental/DataLink/Source/DataLinkEditor/Private/DataLinkEditorLog.h b/Engine/Plugins/Experimental/DataLink/Source/DataLinkEditor/Private/DataLinkEditorLog.h
new file mode 100644
--- /dev/null
+++ b/Engine/Plugins/Experimental/DataLink/Source/DataLinkEditor/Private/DataLinkEditorLog.h
@@ -0,0 +1,7 @@
+// Copyright Epic Games, Inc. All Rights Reserved.
+
+#pragma once
+
+#include "Logging/LogMacros.h"
+
+DECLARE_LOG_CATEGORY_EXTERN(LogDataLinkEditor, Log, All);
diff --git a/Engine/Plugins/Experimental/DataLink/Source/DataLinkEditor/Private/DataLinkEditorModule.cpp b/Engine/Plugins/Experimental/DataLink/Source/DataLinkEditor/Private/DataLinkEditorModule.cpp
new file mode 100644
--- /dev/null
+++ b/Engine/Plugins/Experimental/DataLink/Source/DataLinkEditor/Private/DataLinkEditorModule.cpp
@@ -0,0 +1,45 @@
+// Copyright Epic Games, Inc. All Rights Reserved.
+
+#include "DataLinkEditorModule.h"
+#include "DataLinkEditorLog.h"
+#include "DataLinkGraphCommands.h"
+#include "DataLinkInstance.h"
+#include "DetailsView/DataLinkInstanceCustomization.h"
+#include "Modules/ModuleManager.h"
+#include "PropertyEditorModule.h"
+
+DEFINE_LOG_CATEGORY(LogDataLinkEditor);
+
+IMPLEMENT_MODULE(FDataLinkEditorModule, DataLinkEditor)
+
+void FDataLinkEditorModule::StartupModule()
+{
+	FDataLinkGraphCommands::Register();
+	RegisterDetailCustomizations();
+}
+
+void FDataLinkEditorModule::ShutdownModule()
+{
+	FDataLinkGraphCommands::Unregister();
+	UnregisterDetailCustomizations();
+}
+
+void FDataLinkEditorModule::RegisterDetailCustomizations()
+{
+	FPropertyEditorModule& PropertyEditorModule = FModuleManager::LoadModuleChecked<FPropertyEditorModule>("PropertyEditor");
+
+	PropertyEditorModule.RegisterCustomPropertyTypeLayout(CustomizedTypes.Add_GetRef(FDataLinkInstance::StaticStruct()->GetFName())
+		, FOnGetPropertyTypeCustomizationInstance::CreateStatic(&FDataLinkInstanceCustomization::MakeInstance, /*bGenerateHeader*/true));
+}
+
+void FDataLinkEditorModule::UnregisterDetailCustomizations()
+{
+	if (FPropertyEditorModule* PropertyEditorModule = FModuleManager::GetModulePtr<FPropertyEditorModule>("PropertyEditor"))
+	{
+		for (FName CustomizedType : CustomizedTypes)
+		{
+			PropertyEditorModule->UnregisterCustomPropertyTypeLayout(CustomizedType);
+		}
+		CustomizedTypes.Reset();
+	}
+}
diff --git a/Engine/Plugins/Experimental/DataLink/Source/DataLinkEditor/Private/DataLinkEditorModule.h b/Engine/Plugins/Experimental/DataLink/Source/DataLinkEditor/Private/DataLinkEditorModule.h
new file mode 100644
--- /dev/null
+++ b/Engine/Plugins/Experimental/DataLink/Source/DataLinkEditor/Private/DataLinkEditorModule.h
@@ -0,0 +1,21 @@
+// Copyright Epic Games, Inc. All Rights Reserved.
+
+#pragma once
+
+#include "Containers/Array.h"
+#include "Modules/ModuleInterface.h"
+#include "UObject/NameTypes.h"
+
+class FDataLinkEditorModule : public IModuleInterface
+{
+public:
+	//~ Begin IModuleInterface
+	virtual void StartupModule() override;
+	virtual void ShutdownModule() override;
+	//~ End IModuleInterface
+
+	void RegisterDetailCustomizations();
+	void UnregisterDetailCustomizations();
+
+	TArray<FName> CustomizedTypes;
+};
diff --git a/Engine/Plugins/Experimental/DataLink/Source/DataLinkEditor/Private/DataLinkEditorStyle.cpp b/Engine/Plugins/Experimental/DataLink/Source/DataLinkEditor/Private/DataLinkEditorStyle.cpp
new file mode 100644
--- /dev/null
+++ b/Engine/Plugins/Experimental/DataLink/Source/DataLinkEditor/Private/DataLinkEditorStyle.cpp
@@ -0,0 +1,46 @@
+// Copyright Epic Games, Inc. All Rights Reserved.
+
+#include "DataLinkEditorStyle.h"
+#include "Styling/AppStyle.h"
+#include "Styling/SlateStyleMacros.h"
+#include "Styling/SlateStyleRegistry.h"
+#include "Styling/StyleColors.h"
+
+FDataLinkEditorStyle& FDataLinkEditorStyle::Get()
+{
+	static FDataLinkEditorStyle Instance;
+	return Instance;
+}
+
+FDataLinkEditorStyle::FDataLinkEditorStyle()
+	: FSlateStyleSet(TEXT("DataLinkEditor"))
+{
+	ParentStyleName = FAppStyle::Get().GetStyleSetName();
+
+	const FVector2f Icon20(20.f);
+
+	ContentRootDir = FPaths::EngineContentDir() / TEXT("Editor/Slate");
+	CoreContentRootDir = FPaths::EngineContentDir() / TEXT("Slate");
+
+	// Editor Commands
+	Set("DataLinkGraph.Preview", new IMAGE_BRUSH_SVG("Starship/MainToolbar/simulate", Icon20));
+	Set("DataLinkGraph.ClearPreviewOutput", new IMAGE_BRUSH_SVG("Starship/Common/ResetToDefault", Icon20));
+	Set("DataLinkGraph.ClearPreviewCache", new IMAGE_BRUSH("Icons/GeneralTools/Delete_40x", Icon20));
+
+	Set("CompileStatus.Background.Unknown", new IMAGE_BRUSH_SVG("Starship/Blueprints/CompileStatus_Background", Icon20, FStyleColors::AccentYellow));
+	Set("CompileStatus.Background.Warning", new IMAGE_BRUSH_SVG("Starship/Blueprints/CompileStatus_Background", Icon20, FStyleColors::Warning));
+	Set("CompileStatus.Background.Error"  , new IMAGE_BRUSH_SVG("Starship/Blueprints/CompileStatus_Background", Icon20, FStyleColors::Error));
+	Set("CompileStatus.Background.Good"   , new IMAGE_BRUSH_SVG("Starship/Blueprints/CompileStatus_Background", Icon20, FStyleColors::AccentGreen));
+
+	Set("CompileStatus.Overlay.Unknown", new IMAGE_BRUSH_SVG("Starship/Blueprints/CompileStatus_Unknown_Badge", Icon20, FStyleColors::AccentYellow));
+	Set("CompileStatus.Overlay.Warning", new IMAGE_BRUSH_SVG("Starship/Blueprints/CompileStatus_Warning_Badge", Icon20, FStyleColors::Warning));
+	Set("CompileStatus.Overlay.Error"  , new IMAGE_BRUSH_SVG("Starship/Blueprints/CompileStatus_Fail_Badge"   , Icon20, FStyleColors::Error));
+	Set("CompileStatus.Overlay.Good"   , new IMAGE_BRUSH_SVG("Starship/Blueprints/CompileStatus_Good_Badge"   , Icon20, FStyleColors::AccentGreen));
+
+	FSlateStyleRegistry::RegisterSlateStyle(*this);
+}
+
+FDataLinkEditorStyle::~FDataLinkEditorStyle()
+{
+	FSlateStyleRegistry::UnRegisterSlateStyle(*this);
+}
diff --git a/Engine/Plugins/Experimental/DataLink/Source/DataLinkEditor/Private/DataLinkEditorStyle.h b/Engine/Plugins/Experimental/DataLink/Source/DataLinkEditor/Private/DataLinkEditorStyle.h
new file mode 100644
--- /dev/null
+++ b/Engine/Plugins/Experimental/DataLink/Source/DataLinkEditor/Private/DataLinkEditorStyle.h
@@ -0,0 +1,15 @@
+// Copyright Epic Games, Inc. All Rights Reserved.
+
+#pragma once
+
+#include "Styling/SlateStyle.h"
+
+class FDataLinkEditorStyle final : public FSlateStyleSet
+{
+public:
+	static FDataLinkEditorStyle& Get();
+
+	FDataLinkEditorStyle();
+
+	virtual ~FDataLinkEditorStyle() override;
+};
diff --git a/Engine/Plugins/Experimental/DataLink/Source/DataLinkEditor/Private/DataLinkGraphAssetEditor.cpp b/Engine/Plugins/Experimental/DataLink/Source/DataLinkEditor/Private/DataLinkGraphAssetEditor.cpp
new file mode 100644
--- /dev/null
+++ b/Engine/Plugins/Experimental/DataLink/Source/DataLinkEditor/Private/DataLinkGraphAssetEditor.cpp
@@ -0,0 +1,61 @@
+// Copyright Epic Games, Inc. All Rights Reserved.
+
+#include "DataLinkGraphAssetEditor.h"
+#include "DataLinkEdGraph.h"
+#include "DataLinkGraph.h"
+#include "DataLinkGraphAssetToolkit.h"
+
+void UDataLinkGraphAssetEditor::Initialize(UDataLinkGraph* InDataLinkGraph)
+{
+	DataLinkGraph = InDataLinkGraph;
+
+	if (UDataLinkEdGraph* Graph = GetDataLinkEdGraph())
+	{
+		Graph->InitializeNodes();
+	}
+
+	Super::Initialize();
+}
+
+TSharedPtr<FDataLinkGraphAssetToolkit> UDataLinkGraphAssetEditor::GetToolkit() const
+{
+	if (ToolkitInstance)
+	{
+		return StaticCastSharedRef<FDataLinkGraphAssetToolkit>(ToolkitInstance->AsShared());
+	}
+	return nullptr;
+}
+
+TSharedPtr<FUICommandList> UDataLinkGraphAssetEditor::GetToolkitCommands() const
+{
+	if (ToolkitInstance)
+	{
+		return ToolkitInstance->GetToolkitCommands();
+	}
+	return nullptr;
+}
+
+UDataLinkGraph* UDataLinkGraphAssetEditor::GetDataLinkGraph() const
+{
+	return DataLinkGraph;
+}
+
+UDataLinkEdGraph* UDataLinkGraphAssetEditor::GetDataLinkEdGraph() const
+{
+	if (DataLinkGraph)
+	{
+		return Cast<UDataLinkEdGraph>(DataLinkGraph->GetEdGraph());
+	}
+	return nullptr;
+}
+
+void UDataLinkGraphAssetEditor::GetObjectsToEdit(TArray<UObject*>& InObjectsToEdit)
+{
+	Super::GetObjectsToEdit(InObjectsToEdit);
+	InObjectsToEdit.Add(GetDataLinkGraph());
+}
+
+TSharedPtr<FBaseAssetToolkit> UDataLinkGraphAssetEditor::CreateToolkit()
+{
+	return MakeShared<FDataLinkGraphAssetToolkit>(this);
+}
diff --git a/Engine/Plugins/Experimental/DataLink/Source/DataLinkEditor/Private/DataLinkGraphAssetEditor.h b/Engine/Plugins/Experimental/DataLink/Source/DataLinkEditor/Private/DataLinkGraphAssetEditor.h
new file mode 100644
--- /dev/null
+++ b/Engine/Plugins/Experimental/DataLink/Source/DataLinkEditor/Private/DataLinkGraphAssetEditor.h
@@ -0,0 +1,37 @@
+// Copyright Epic Games, Inc. All Rights Reserved.
+
+#pragma once
+
+#include "Tools/UAssetEditor.h"
+#include "DataLinkGraphAssetEditor.generated.h"
+
+class FDataLinkGraphAssetToolkit;
+class UDataLinkEdGraph;
+class UDataLinkGraph;
+
+UCLASS(Transient)
+class UDataLinkGraphAssetEditor : public UAssetEditor
+{
+	GENERATED_BODY()
+
+public:
+	void Initialize(UDataLinkGraph* InDataLinkGraph);
+
+	TSharedPtr<FDataLinkGraphAssetToolkit> GetToolkit() const;
+
+	TSharedPtr<FUICommandList> GetToolkitCommands() const;
+
+	UDataLinkGraph* GetDataLinkGraph() const;
+
+	UDataLinkEdGraph* GetDataLinkEdGraph() const;
+
+protected:
+	//~ Begin UAssetEditor
+	virtual void GetObjectsToEdit(TArray<UObject*>& InObjectsToEdit) override;
+	virtual TSharedPtr<FBaseAssetToolkit> CreateToolkit() override;
+	//~ End UAssetEditor
+
+private:
+	UPROPERTY()
+	TObjectPtr<UDataLinkGraph> DataLinkGraph;
+};
diff --git a/Engine/Plugins/Experimental/DataLink/Source/DataLinkEditor/Private/DataLinkGraphAssetToolkit.cpp b/Engine/Plugins/Experimental/DataLink/Source/DataLinkEditor/Private/DataLinkGraphAssetToolkit.cpp
new file mode 100644
--- /dev/null
+++ b/Engine/Plugins/Experimental/DataLink/Source/DataLinkEditor/Private/DataLinkGraphAssetToolkit.cpp
@@ -0,0 +1,173 @@
+// Copyright Epic Games, Inc. All Rights Reserved.
+
+#include "DataLinkGraphAssetToolkit.h"
+#include "Compiler/DataLinkGraphCompilerTool.h"
+#include "DataLinkEdGraph.h"
+#include "DataLinkGraph.h"
+#include "DataLinkGraphAssetEditor.h"
+#include "DataLinkGraphEditorMenuContext.h"
+#include "Framework/Docking/LayoutExtender.h"
+#include "Graph/DataLinkGraphEditorTool.h"
+#include "Modules/ModuleManager.h"
+#include "Preview/DataLinkPreviewTool.h"
+#include "PropertyEditorModule.h"
+#include "ToolMenus.h"
+
+#define LOCTEXT_NAMESPACE "DataLinkGraphAssetToolkit"
+
+FDataLinkGraphAssetToolkit::FDataLinkGraphAssetToolkit(UDataLinkGraphAssetEditor* InAssetEditor)
+	: FBaseAssetToolkit(InAssetEditor)
+	, AssetEditor(InAssetEditor)
+	, GraphTool(MakeShared<FDataLinkGraphEditorTool>(InAssetEditor))
+	, CompilerTool(MakeShared<FDataLinkGraphCompilerTool>(InAssetEditor))
+	, PreviewTool(MakeShared<FDataLinkPreviewTool>(InAssetEditor))
+{
+	LayoutAppendix = TEXT("DataLinkGraphAssetEditor");
+
+	StandaloneDefaultLayout = FTabManager::NewLayout(FName(TEXT("Standalone_Layout_V_1_0_") + LayoutAppendix))
+		->AddArea
+		(
+			FTabManager::NewPrimaryArea()
+			->SetOrientation(Orient_Horizontal)
+			->Split
+			(
+				FTabManager::NewStack()
+				->SetSizeCoefficient(0.7f)
+				->AddTab(FDataLinkGraphEditorTool::GraphEditorTabID, ETabState::OpenedTab)
+			)
+			->Split
+			(
+				FTabManager::NewSplitter()
+				->SetSizeCoefficient(0.3f)
+				->SetOrientation(Orient_Vertical)
+				->Split
+				(
+					FTabManager::NewStack()
+					->SetSizeCoefficient(0.5f)
+					->AddTab(DetailsTabID, ETabState::OpenedTab)
+				)
+				->Split
+				(
+					FTabManager::NewStack()
+					->SetSizeCoefficient(0.5f)
+					->AddTab(FDataLinkPreviewTool::PreviewTabID, ETabState::OpenedTab)
+				)
+			)
+		);
+}
+
+TSharedPtr<IDetailsView> FDataLinkGraphAssetToolkit::GetDetailsView() const
+{
+	return DetailsView;
+}
+
+void FDataLinkGraphAssetToolkit::RegisterToolbar()
+{
+	FName ParentName;
+	const FName MenuName = GetToolMenuToolbarName(ParentName);
+
+	UToolMenus* ToolMenus = UToolMenus::Get();
+	if (!ToolMenus || ToolMenus->IsMenuRegistered(MenuName))
+	{
+		return;
+	}
+
+	FToolMenuOwnerScoped ToolMenuOwnerScope(this);
+
+	UToolMenu* const ToolbarMenu = ToolMenus->RegisterMenu(MenuName, ParentName, EMultiBoxType::ToolBar);
+	check(ToolbarMenu);
+
+	CompilerTool->ExtendMenu(ToolbarMenu);
+}
+
+void FDataLinkGraphAssetToolkit::InitToolMenuContext(FToolMenuContext& InMenuContext)
+{
+	FBaseAssetToolkit::InitToolMenuContext(InMenuContext);
+
+	UDataLinkGraphEditorMenuContext* Context = NewObject<UDataLinkGraphEditorMenuContext>();
+	Context->ToolkitWeak = SharedThis(this);
+	InMenuContext.AddObject(Context);
+}
+
+void FDataLinkGraphAssetToolkit::CreateWidgets()
+{
+	// Skip FBaseAssetToolkit to avoid creating a Viewport Tab/Client/etc
+	RegisterToolbar();
+	CreateEditorModeManager();
+
+	LayoutExtender = MakeShared<FLayoutExtender>();
+
+	GraphTool->CreateWidgets();
+	PreviewTool->CreateWidgets();
+
+	DetailsView = CreateDetailsView();
+}
+
+TSharedRef<IDetailsView> FDataLinkGraphAssetToolkit::CreateDetailsView()
+{
+	FPropertyEditorModule& PropertyEditorModule = FModuleManager::GetModuleChecked<FPropertyEditorModule>("PropertyEditor");
+
+	FDetailsViewArgs DetailsViewArgs;
+	DetailsViewArgs.DefaultsOnlyVisibility = EEditDefaultsOnlyNodeVisibility::Automatic;
+	DetailsViewArgs.NameAreaSettings = FDetailsViewArgs::HideNameArea;
+	DetailsViewArgs.bHideSelectionTip = true;
+
+	return PropertyEditorModule.CreateDetailView(DetailsViewArgs);
+}
+
+void FDataLinkGraphAssetToolkit::PostInitAssetEditor()
+{
+	FBaseAssetToolkit::PostInitAssetEditor();
+
+	if (UDataLinkEdGraph* EdGraph = AssetEditor->GetDataLinkEdGraph())
+	{
+		EdGraph->AddOnGraphChangedHandler(FOnGraphChanged::FDelegate::CreateSP(this, &FDataLinkGraphAssetToolkit::OnGraphChanged));
+	}
+}
+
+void FDataLinkGraphAssetToolkit::MapToolkitCommands()
+{
+	FBaseAssetToolkit::MapToolkitCommands();
+
+	GraphTool->BindCommands(ToolkitCommands);
+	CompilerTool->BindCommands(ToolkitCommands);
+	PreviewTool->BindCommands(ToolkitCommands);
+}
+
+FName FDataLinkGraphAssetToolkit::GetToolkitFName() const
+{
+	return TEXT("FDataLinkGraphAssetToolkit");
+}
+
+void FDataLinkGraphAssetToolkit::RegisterTabSpawners(const TSharedRef<FTabManager>& InTabManager)
+{
+	// Skip FBaseAssetToolkit to not register a viewport tab.
+	FAssetEditorToolkit::RegisterTabSpawners(InTabManager);
+
+	GraphTool->RegisterTabSpawners(InTabManager, AssetEditorTabsCategory);
+	PreviewTool->RegisterTabSpawners(InTabManager, AssetEditorTabsCategory);
+
+	InTabManager->RegisterTabSpawner(DetailsTabID, FOnSpawnTab::CreateSP(this, &FDataLinkGraphAssetToolkit::SpawnTab_Details))
+		.SetDisplayName(LOCTEXT("Details", "Details"))
+		.SetGroup(AssetEditorTabsCategory.ToSharedRef())
+		.SetIcon(FSlateIcon(FAppStyle::GetAppStyleSetName(), "LevelEditor.Tabs.Details"));
+}
+
+void FDataLinkGraphAssetToolkit::UnregisterTabSpawners(const TSharedRef<FTabManager>& InTabManager)
+{
+	FAssetEditorToolkit::UnregisterTabSpawners(InTabManager);
+
+	GraphTool->UnregisterTabSpawners(InTabManager);
+	PreviewTool->UnregisterTabSpawners(InTabManager);
+	InTabManager->UnregisterTabSpawner(DetailsTabID);
+}
+
+void FDataLinkGraphAssetToolkit::OnGraphChanged(const FEdGraphEditAction& InAction)
+{
+	if (UDataLinkEdGraph* EdGraph = AssetEditor->GetDataLinkEdGraph())
+	{
+		EdGraph->DirtyGraph();
+	}
+}
+
+#undef LOCTEXT_NAMESPACE
diff --git a/Engine/Plugins/Experimental/DataLink/Source/DataLinkEditor/Private/DataLinkGraphAssetToolkit.h b/Engine/Plugins/Experimental/DataLink/Source/DataLinkEditor/Private/DataLinkGraphAssetToolkit.h
new file mode 100644
--- /dev/null
+++ b/Engine/Plugins/Experimental/DataLink/Source/DataLinkEditor/Private/DataLinkGraphAssetToolkit.h
@@ -0,0 +1,59 @@
+// Copyright Epic Games, Inc. All Rights Reserved.
+
+#pragma once
+
+#include "Tools/BaseAssetToolkit.h"
+
+class FDataLinkGraphCompilerTool;
+class FDataLinkGraphEditorTool;
+class FDataLinkPreviewTool;
+class SGraphEditor;
+class UDataLinkGraphAssetEditor;
+struct FEdGraphEditAction;
+
+class FDataLinkGraphAssetToolkit : public FBaseAssetToolkit
+{
+public:
+	explicit FDataLinkGraphAssetToolkit(UDataLinkGraphAssetEditor* InAssetEditor);
+
+	TSharedPtr<IDetailsView> GetDetailsView() const;
+
+	//~ Begin FBaseAssetToolkit
+	virtual void RegisterToolbar() override;
+	virtual void InitToolMenuContext(FToolMenuContext& InMenuContext) override;
+	virtual void CreateWidgets() override;
+	virtual TSharedRef<IDetailsView> CreateDetailsView() override;
+	virtual void PostInitAssetEditor() override;
+	//~ End FBaseAssetToolkit
+
+	//~ Begin FAssetEditorToolkit
+	virtual void MapToolkitCommands() override;
+	//~ End FAssetEditorToolkit
+
+	//~ Begin IToolkit
+	virtual FName GetToolkitFName() const override;
+	virtual void RegisterTabSpawners(const TSharedRef<FTabManager>& InTabManager) override;
+	virtual void UnregisterTabSpawners(const TSharedRef<FTabManager>& InTabManager) override;
+	//~ End IToolkit
+
+	const FDataLinkGraphCompilerTool& GetCompilerTool() const
+	{
+		return CompilerTool.Get();
+	}
+
+	const FDataLinkPreviewTool& GetPreviewTool() const
+	{
+		return PreviewTool.Get();
+	}
+
+private:
+	void OnGraphChanged(const FEdGraphEditAction& InAction);
+
+	TObjectPtr<UDataLinkGraphAssetEditor> AssetEditor;
+
+	TSharedRef<FDataLinkGraphEditorTool> GraphTool;
+
+	TSharedRef<FDataLinkGraphCompilerTool> CompilerTool;
+
+	TSharedRef<FDataLinkPreviewTool> PreviewTool;
+};
diff --git a/Engine/Plugins/Experimental/DataLink/Source/DataLinkEditor/Private/DataLinkGraphCommands.cpp b/Engine/Plugins/Experimental/DataLink/Source/DataLinkEditor/Private/DataLinkGraphCommands.cpp
new file mode 100644
--- /dev/null
+++ b/Engine/Plugins/Experimental/DataLink/Source/DataLinkEditor/Private/DataLinkGraphCommands.cpp
@@ -0,0 +1,44 @@
+// Copyright Epic Games, Inc. All Rights Reserved.
+
+#include "DataLinkGraphCommands.h"
+#include "DataLinkEditorStyle.h"
+#include "Styling/AppStyle.h"
+
+#define LOCTEXT_NAMESPACE "DataLinkGraphCommands"
+
+FDataLinkGraphCommands::FDataLinkGraphCommands()
+	: TCommands(TEXT("DataLinkGraph")
+	, LOCTEXT("DataLinkGraph", "Data Link Graph")
+	, NAME_None
+	, FDataLinkEditorStyle::Get().GetStyleSetName())
+{
+}
+
+void FDataLinkGraphCommands::RegisterCommands()
+{
+	UI_COMMAND(Compile
+		, "Compile"
+		, "Compiles the data link graph"
+		, EUserInterfaceActionType::Button
+		, FInputChord());
+
+	UI_COMMAND(Preview
+		, "Preview"
+		, "Triggers execution of the data link graph with the given input data"
+		, EUserInterfaceActionType::Button
+		, FInputChord());
+
+	UI_COMMAND(ClearPreviewOutput
+		, "Clear Output"
+		, "Clears the output data from the last execution."
+		, EUserInterfaceActionType::Button
+		, FInputChord());
+
+	UI_COMMAND(ClearPreviewCache
+		, "Clear Cache"
+		, "Clears the cached node data from previous executions."
+		, EUserInterfaceActionType::Button
+		, FInputChord());
+}
+
+#undef LOCTEXT_NAMESPACE
diff --git a/Engine/Plugins/Experimental/DataLink/Source/DataLinkEditor/Private/DataLinkGraphCommands.h b/Engine/Plugins/Experimental/DataLink/Source/DataLinkEditor/Private/DataLinkGraphCommands.h
new file mode 100644
--- /dev/null
+++ b/Engine/Plugins/Experimental/DataLink/Source/DataLinkEditor/Private/DataLinkGraphCommands.h
@@ -0,0 +1,23 @@
+// Copyright Epic Games, Inc. All Rights Reserved.
+
+#pragma once
+
+#include "Framework/Commands/Commands.h"
+
+class FDataLinkGraphCommands : public TCommands<FDataLinkGraphCommands>
+{
+public:
+	FDataLinkGraphCommands();
+
+	//~ Begin TCommands
+	virtual void RegisterCommands() override;
+	//~ End TCommands
+
+	TSharedPtr<FUICommandInfo> Compile;
+
+	TSharedPtr<FUICommandInfo> Preview;
+
+	TSharedPtr<FUICommandInfo> ClearPreviewOutput;
+
+	TSharedPtr<FUICommandInfo> ClearPreviewCache;
+};
diff --git a/Engine/Plugins/Experimental/DataLink/Source/DataLinkEditor/Private/DataLinkGraphEditorMenuContext.cpp b/Engine/Plugins/Experimental/DataLink/Source/DataLinkEditor/Private/DataLinkGraphEditorMenuContext.cpp
new file mode 100644
--- /dev/null
+++ b/Engine/Plugins/Experimental/DataLink/Source/DataLinkEditor/Private/DataLinkGraphEditorMenuContext.cpp
@@ -0,0 +1,52 @@
+// Copyright Epic Games, Inc. All Rights Reserved.
+
+#include "DataLinkGraphEditorMenuContext.h"
+#include "DataLinkGraphAssetToolkit.h"
+#include "Preview/DataLinkPreviewData.h"
+#include "Preview/DataLinkPreviewTool.h"
+#include "StructUtils/StructView.h"
+
+FConstStructView UDataLinkGraphEditorMenuContext::FindPreviewOutputData() const
+{
+	TSharedPtr<FDataLinkGraphAssetToolkit> Toolkit = ToolkitWeak.Pin();
+	if (!Toolkit.IsValid())
+	{
+		return FConstStructView();
+	}
+
+	const UDataLinkPreviewData* const PreviewData = Toolkit->GetPreviewTool().GetPreviewData();
+	if (!PreviewData)
+	{
+		return FConstStructView();
+	}
+
+	return PreviewData->OutputData;
+}
+
+FString UDataLinkGraphEditorMenuContext::GetAssetPath() const
+{
+	TSharedPtr<FDataLinkGraphAssetToolkit> Toolkit = ToolkitWeak.Pin();
+	if (!Toolkit.IsValid())
+	{
+		return FString();
+	}
+
+	const TArray<UObject*>* Objects = Toolkit->GetObjectsCurrentlyBeingEdited();
+	if (!Objects || Objects->IsEmpty())
+	{
+		return FString();
+	}
+
+	UObject* const Object = (*Objects)[0];
+	if (!Object)
+	{
+		return FString();
+	}
+
+	if (const UPackage* Package = Object->GetPackage())
+	{
+		return FPaths::GetPath(Package->GetPathName());
+	}
+
+	return FString();
+}
diff --git a/Engine/Plugins/Experimental/DataLink/Source/DataLinkEditor/Private/DataLinkGraphEditorMenuContext.h b/Engine/Plugins/Experimental/DataLink/Source/DataLinkEditor/Private/DataLinkGraphEditorMenuContext.h
new file mode 100644
--- /dev/null
+++ b/Engine/Plugins/Experimental/DataLink/Source/DataLinkEditor/Private/DataLinkGraphEditorMenuContext.h
@@ -0,0 +1,24 @@
+// Copyright Epic Games, Inc. All Rights Reserved.
+
+#pragma once
+
+#include "IDataLinkEditorMenuContext.h"
+#include "Templates/SharedPointer.h"
+#include "UObject/Object.h"
+#include "DataLinkGraphEditorMenuContext.generated.h"
+
+class FDataLinkGraphAssetToolkit;
+
+UCLASS()
+class UDataLinkGraphEditorMenuContext : public UObject, public IDataLinkEditorMenuContext
+{
+	GENERATED_BODY()
+
+public:
+	//~ Begin IDataLinkEditorMenuContext
+	virtual FConstStructView FindPreviewOutputData() const override;
+	virtual FString GetAssetPath() const override;
+	//~ End IDataLinkEditorMenuContext
+
+	TWeakPtr<FDataLinkGraphAssetToolkit> ToolkitWeak;
+};
diff --git a/Engine/Plugins/Experimental/DataLink/Source/DataLinkEditor/Private/DataLinkGraphFactory.cpp b/Engine/Plugins/Experimental/DataLink/Source/DataLinkEditor/Private/DataLinkGraphFactory.cpp
new file mode 100644
--- /dev/null
+++ b/Engine/Plugins/Experimental/DataLink/Source/DataLinkEditor/Private/DataLinkGraphFactory.cpp
@@ -0,0 +1,44 @@
+// Copyright Epic Games, Inc. All Rights Reserved.
+
+#include "DataLinkGraphFactory.h"
+#include "AssetTypeCategories.h"
+#include "DataLinkEdGraph.h"
+#include "DataLinkEdGraphSchema.h"
+#include "DataLinkGraph.h"
+#include "Kismet2/BlueprintEditorUtils.h"
+
+UDataLinkGraphFactory::UDataLinkGraphFactory()
+{
+	SupportedClass = UDataLinkGraph::StaticClass();
+	bCreateNew = true;
+	bEditorImport = false;
+	bEditAfterNew = true;
+}
+
+uint32 UDataLinkGraphFactory::GetMenuCategories() const
+{
+	return EAssetTypeCategories::Misc;
+}
+
+UObject* UDataLinkGraphFactory::FactoryCreateNew(UClass* InClass, UObject* InParent, FName InName, EObjectFlags InFlags, UObject* InContext, FFeedbackContext* InWarn)
+{
+	if (!ensure(SupportedClass == InClass))
+	{
+		return nullptr;
+	}
+
+	UDataLinkGraph* DataLinkGraph = NewObject<UDataLinkGraph>(InParent, InName, InFlags);
+	check(DataLinkGraph);
+
+	DataLinkGraph->EdGraph = FBlueprintEditorUtils::CreateNewGraph(DataLinkGraph
+		, NAME_None
+		, UDataLinkEdGraph::StaticClass()
+		, UDataLinkEdGraphSchema::StaticClass());
+	check(DataLinkGraph->EdGraph);
+
+	const UEdGraphSchema* Schema = DataLinkGraph->EdGraph->GetSchema();
+	check(Schema);
+	Schema->CreateDefaultNodesForGraph(*DataLinkGraph->EdGraph);
+
+	return DataLinkGraph;
+}
diff --git a/Engine/Plugins/Experimental/DataLink/Source/DataLinkEditor/Private/DataLinkGraphFactory.h b/Engine/Plugins/Experimental/DataLink/Source/DataLinkEditor/Private/DataLinkGraphFactory.h
new file mode 100644
--- /dev/null
+++ b/Engine/Plugins/Experimental/DataLink/Source/DataLinkEditor/Private/DataLinkGraphFactory.h
@@ -0,0 +1,20 @@
+// Copyright Epic Games, Inc. All Rights Reserved.
+
+#pragma once
+
+#include "Factories/Factory.h"
+#include "DataLinkGraphFactory.generated.h"
+
+UCLASS()
+class UDataLinkGraphFactory : public UFactory
+{
+	GENERATED_BODY()
+
+public:
+	UDataLinkGraphFactory();
+
+	//~ Begin UFactory
+	virtual uint32 GetMenuCategories() const override;
+	virtual UObject* FactoryCreateNew(UClass* InClass, UObject* InParent, FName InName, EObjectFlags InFlags, UObject* InContext, FFeedbackContext* InWarn) override;
+	//~ End UFactory
+};
diff --git a/Engine/Plugins/Experimental/DataLink/Source/DataLinkEditor/Private/DetailsView/DataLinkInstanceCustomization.cpp b/Engine/Plugins/Experimental/DataLink/Source/DataLinkEditor/Private/DetailsView/DataLinkInstanceCustomization.cpp
new file mode 100644
--- /dev/null
+++ b/Engine/Plugins/Experimental/DataLink/Source/DataLinkEditor/Private/DetailsView/DataLinkInstanceCustomization.cpp
@@ -0,0 +1,152 @@
+// Copyright Epic Games, Inc. All Rights Reserved.
+
+#include "DataLinkInstanceCustomization.h"
+#include "DataLinkGraph.h"
+#include "DataLinkInstance.h"
+#include "DataLinkInstancedStructNodeBuilder.h"
+#include "DataLinkPin.h"
+#include "DataLinkPinReference.h"
+#include "DetailWidgetRow.h"
+#include "IDetailChildrenBuilder.h"
+#include "PropertyCustomizationHelpers.h"
+
+FDataLinkInstanceCustomization::FDataLinkInstanceCustomization(bool bInGenerateHeader)
+	: bGenerateHeader(bInGenerateHeader)
+{
+}
+
+FDataLinkInstanceCustomization::~FDataLinkInstanceCustomization()
+{
+	UDataLinkGraph::OnGraphCompiled().RemoveAll(this);
+}
+
+void FDataLinkInstanceCustomization::CustomizeHeader(TSharedRef<IPropertyHandle> InPropertyHandle, FDetailWidgetRow& InHeaderRow, IPropertyTypeCustomizationUtils& InCustomizationUtils)
+{
+	UDataLinkGraph::OnGraphCompiled().AddSP(this, &FDataLinkInstanceCustomization::OnGraphCompiled);
+
+	DataLinkGraphHandle = InPropertyHandle->GetChildHandle(GET_MEMBER_NAME_CHECKED(FDataLinkInstance, DataLinkGraph));
+	check(DataLinkGraphHandle.IsValid());
+	DataLinkGraphHandle->MarkHiddenByCustomization();
+	DataLinkGraphHandle->SetOnPropertyValueChanged(FSimpleDelegate::CreateSP(this, &FDataLinkInstanceCustomization::OnGraphChanged));
+
+	InputDataHandle = InPropertyHandle->GetChildHandle(GET_MEMBER_NAME_CHECKED(FDataLinkInstance, InputData));
+	check(InputDataHandle.IsValid());
+	InputDataHandle->MarkHiddenByCustomization();
+
+	if (bGenerateHeader)
+	{
+		InHeaderRow
+			.NameContent()
+			[
+				InPropertyHandle->CreatePropertyNameWidget()
+			]
+			.ValueContent()
+			[
+				DataLinkGraphHandle->CreatePropertyValueWidget()
+			];
+	}
+
+	UpdateInputData();
+}
+
+void FDataLinkInstanceCustomization::CustomizeChildren(TSharedRef<IPropertyHandle> InPropertyHandle, IDetailChildrenBuilder& InChildBuilder, IPropertyTypeCustomizationUtils& InCustomizationUtils)
+{
+	check(InputDataHandle.IsValid());
+
+	UpdateInputData();
+
+	TSharedRef<FDetailArrayBuilder> ArrayBuilder = MakeShared<FDetailArrayBuilder>(InputDataHandle.ToSharedRef()
+		, /*bGenerateHeader*/false
+		, /*bDisplayResetToDefault*/false
+		, /*DisplayElementNum*/false);
+
+	ArrayBuilder->OnGenerateArrayElementWidget(FOnGenerateArrayElementWidget::CreateLambda(
+		[&InputDisplayNames = InputDisplayNames](TSharedRef<IPropertyHandle> InElementHandle, int32 InIndex, IDetailChildrenBuilder& InChildBuilder)
+		{
+			if (InputDisplayNames.IsValidIndex(InIndex))
+			{
+				InElementHandle->SetPropertyDisplayName(InputDisplayNames[InIndex]);
+			}
+			InChildBuilder.AddCustomBuilder(MakeShared<FDataLinkInstancedStructNodeBuilder>(InElementHandle));
+		}));
+
+	InChildBuilder.AddCustomBuilder(ArrayBuilder);
+}
+
+void FDataLinkInstanceCustomization::OnGraphCompiled(UDataLinkGraph* InDataLinkGraph)
+{
+	if (InDataLinkGraph == GetDataLinkGraph())
+	{
+		OnGraphChanged();
+	}
+}
+
+UDataLinkGraph* FDataLinkInstanceCustomization::GetDataLinkGraph() const
+{
+	UObject* DataLink = nullptr;
+	if (DataLinkGraphHandle.IsValid() && DataLinkGraphHandle->GetValue(DataLink) == FPropertyAccess::Success)
+	{
+		return Cast<UDataLinkGraph>(DataLink);
+	}
+	return nullptr;
+}
+
+TArray<FInstancedStruct>* FDataLinkInstanceCustomization::GetInputData() const
+{
+	void* RawData = nullptr;
+	if (InputDataHandle.IsValid() && InputDataHandle->GetValueData(RawData) == FPropertyAccess::Success)
+	{
+		return static_cast<TArray<FInstancedStruct>*>(RawData);
+	}
+	return nullptr;
+}
+
+void FDataLinkInstanceCustomization::OnGraphChanged()
+{
+	if (!InputDataHandle.IsValid())
+	{
+		return;
+	}
+
+	InputDataHandle->NotifyPreChange();
+	UpdateInputData();
+	InputDataHandle->NotifyPostChange(EPropertyChangeType::ValueSet);
+}
+
+void FDataLinkInstanceCustomization::UpdateInputData()
+{
+	TArray<FInstancedStruct>* InputData = GetInputData();
+	if (!InputData)
+	{
+		return;
+	}
+
+	UDataLinkGraph* DataLinkGraph = GetDataLinkGraph();
+	if (!DataLinkGraph)
+	{
+		InputDisplayNames.Reset();
+		InputData->Reset();
+		return;
+	}
+
+	const int32 InputPinCount = DataLinkGraph->GetInputPinCount();
+	InputDisplayNames.SetNum(InputPinCount);
+	InputData->SetNum(InputPinCount);
+
+	int32 Index = 0;
+	DataLinkGraph->ForEachInputPin(
+		[InputData, &Index, &InputDisplayNames = InputDisplayNames](const FDataLinkPinReference& InPinReference)->bool
+		{
+			InputDisplayNames[Index] = InPinReference.Pin->GetDisplayName();
+
+			FInstancedStruct& InstanceData = (*InputData)[Index];
+			if (InstanceData.GetScriptStruct() != InPinReference.Pin->Struct)
+			{
+				// Input Struct can be null, FInstancedStruct will just be reset
+				InstanceData.InitializeAs(InPinReference.Pin->Struct);
+			}
+
+			++Index;
+			return true;
+		});
+}
diff --git a/Engine/Plugins/Experimental/DataLink/Source/DataLinkEditor/Private/DetailsView/DataLinkInstanceCustomization.h b/Engine/Plugins/Experimental/DataLink/Source/DataLinkEditor/Private/DetailsView/DataLinkInstanceCustomization.h
new file mode 100644
--- /dev/null
+++ b/Engine/Plugins/Experimental/DataLink/Source/DataLinkEditor/Private/DetailsView/DataLinkInstanceCustomization.h
@@ -0,0 +1,43 @@
+// Copyright Epic Games, Inc. All Rights Reserved.
+
+#pragma once
+
+#include "IPropertyTypeCustomization.h"
+
+class UDataLinkGraph;
+struct FInstancedStruct;
+
+class FDataLinkInstanceCustomization : public IPropertyTypeCustomization
+{
+public:
+	static TSharedRef<IPropertyTypeCustomization> MakeInstance(bool bInGenerateHeader)
+	{
+		return MakeShared<FDataLinkInstanceCustomization>(bInGenerateHeader);
+	}
+
+	explicit FDataLinkInstanceCustomization(bool bInGenerateHeader);
+
+	virtual ~FDataLinkInstanceCustomization() override;
+
+protected:
+	//~ Begin IPropertyTypeCustomization
+	virtual void CustomizeHeader(TSharedRef<IPropertyHandle> InPropertyHandle, FDetailWidgetRow& InHeaderRow, IPropertyTypeCustomizationUtils& InCustomizationUtils) override;
+	virtual void CustomizeChildren(TSharedRef<IPropertyHandle> InPropertyHandle, IDetailChildrenBuilder& InChildBuilder, IPropertyTypeCustomizationUtils& InCustomizationUtils) override;
+	//~ End IPropertyTypeCustomization
+
+	void OnGraphCompiled(UDataLinkGraph* InDataLinkGraph);
+
+	UDataLinkGraph* GetDataLinkGraph() const;
+
+	TArray<FInstancedStruct>* GetInputData() const;
+
+	void OnGraphChanged();
+	void UpdateInputData();
+
+	TSharedPtr<IPropertyHandle> DataLinkGraphHandle;
+	TSharedPtr<IPropertyHandle> InputDataHandle;
+
+	TArray<FText> InputDisplayNames;
+
+	bool bGenerateHeader;
+};
diff --git a/Engine/Plugins/Experimental/DataLink/Source/DataLinkEditor/Private/DetailsView/DataLinkInstancedStructNodeBuilder.cpp b/Engine/Plugins/Experimental/DataLink/Source/DataLinkEditor/Private/DetailsView/DataLinkInstancedStructNodeBuilder.cpp
new file mode 100644
--- /dev/null
+++ b/Engine/Plugins/Experimental/DataLink/Source/DataLinkEditor/Private/DetailsView/DataLinkInstancedStructNodeBuilder.cpp
@@ -0,0 +1,76 @@
+// Copyright Epic Games, Inc. All Rights Reserved.
+
+#include "DataLinkInstancedStructNodeBuilder.h"
+#include "DetailLayoutBuilder.h"
+#include "DetailWidgetRow.h"
+#include "IDetailChildrenBuilder.h"
+#include "Widgets/SBoxPanel.h"
+#include "Widgets/Text/STextBlock.h"
+
+FDataLinkInstancedStructNodeBuilder::FDataLinkInstancedStructNodeBuilder(const TSharedRef<IPropertyHandle>& InInputDataElementHandle)
+	: FInstancedStructDataDetails(InInputDataElementHandle)
+	, InputDataElementHandle(InInputDataElementHandle)
+{
+}
+
+void FDataLinkInstancedStructNodeBuilder::GenerateHeaderRowContent(FDetailWidgetRow& InNodeRow)
+{
+	FInstancedStructDataDetails::GenerateHeaderRowContent(InNodeRow);
+
+	InNodeRow
+		.ShouldAutoExpand(true)
+		.NameContent()
+		[
+			SNew(SHorizontalBox)
+			+ SHorizontalBox::Slot()
+			.HAlign(HAlign_Left)
+			.VAlign(VAlign_Center)
+			.AutoWidth()
+			[
+				SNew(STextBlock)
+				.Text(INVTEXT("["))
+				.Font(IDetailLayoutBuilder::GetDetailFont())
+				.ColorAndOpacity(FSlateColor::UseSubduedForeground())
+			]
+			+ SHorizontalBox::Slot()
+			.HAlign(HAlign_Left)
+			.VAlign(VAlign_Center)
+			.Padding(3, 0, 0, 0)
+			.AutoWidth()
+			[
+				SNew(STextBlock)
+				.Text(FText::AsNumber(InputDataElementHandle->GetArrayIndex()))
+				.Font(IDetailLayoutBuilder::GetDetailFont())
+			]
+			+ SHorizontalBox::Slot()
+			.HAlign(HAlign_Left)
+			.VAlign(VAlign_Center)
+			.Padding(3, 0, 0, 0)
+			.AutoWidth()
+			[
+				SNew(STextBlock)
+				.Text(INVTEXT("]"))
+				.Font(IDetailLayoutBuilder::GetDetailFont())
+				.ColorAndOpacity(FSlateColor::UseSubduedForeground())
+			]
+			+ SHorizontalBox::Slot()
+			.HAlign(HAlign_Left)
+			.VAlign(VAlign_Center)
+			.Padding(5, 0, 0, 0)
+			.AutoWidth()
+			[
+				SNew(STextBlock)
+				.Text(InputDataElementHandle->GetPropertyDisplayName())
+				.Font(IDetailLayoutBuilder::GetDetailFont())
+			]
+		];
+}
+
+void FDataLinkInstancedStructNodeBuilder::Tick(float InDeltaTime)
+{
+	// TODO: Logic in FInstancedStructDataDetails::Tick calls EnumerateConstRawData which does not check that the underlying handle still has a valid property node
+	if (InputDataElementHandle->IsValidHandle())
+	{
+		FInstancedStructDataDetails::Tick(InDeltaTime);
+	}
+}
diff --git a/Engine/Plugins/Experimental/DataLink/Source/DataLinkEditor/Private/DetailsView/DataLinkInstancedStructNodeBuilder.h b/Engine/Plugins/Experimental/DataLink/Source/DataLinkEditor/Private/DetailsView/DataLinkInstancedStructNodeBuilder.h
new file mode 100644
--- /dev/null
+++ b/Engine/Plugins/Experimental/DataLink/Source/DataLinkEditor/Private/DetailsView/DataLinkInstancedStructNodeBuilder.h
@@ -0,0 +1,19 @@
+// Copyright Epic Games, Inc. All Rights Reserved.
+
+#pragma once
+
+#include "InstancedStructDetails.h"
+
+class FDataLinkInstancedStructNodeBuilder : public FInstancedStructDataDetails
+{
+public:
+	explicit FDataLinkInstancedStructNodeBuilder(const TSharedRef<IPropertyHandle>& InInputDataElementHandle);
+
+	//~ Begin IDetailCustomNodeBuilder
+	virtual void GenerateHeaderRowContent(FDetailWidgetRow& InNodeRow) override;
+	virtual void Tick(float InDeltaTime) override;
+	//~ End IDetailCustomNodeBuilder
+
+private:
+	TSharedRef<IPropertyHandle> InputDataElementHandle;
+};
diff --git a/Engine/Plugins/Experimental/DataLink/Source/DataLinkEditor/Private/Graph/DataLinkGraphEditorTool.cpp b/Engine/Plugins/Experimental/DataLink/Source/DataLinkEditor/Private/Graph/DataLinkGraphEditorTool.cpp
new file mode 100644
--- /dev/null
+++ b/Engine/Plugins/Experimental/DataLink/Source/DataLinkEditor/Private/Graph/DataLinkGraphEditorTool.cpp
@@ -0,0 +1,364 @@
+// Copyright Epic Games, Inc. All Rights Reserved.
+
+#include "DataLinkGraphEditorTool.h"
+#include "BlueprintActionDatabase.h"
+#include "Customizations/MathStructProxyCustomizations.h"
+#include "DataLinkEdGraph.h"
+#include "DataLinkGraphAssetEditor.h"
+#include "DataLinkGraphAssetToolkit.h"
+#include "EdGraphUtilities.h"
+#include "Framework/Commands/GenericCommands.h"
+#include "Framework/Docking/TabManager.h"
+#include "GraphEditor.h"
+#include "HAL/PlatformApplicationMisc.h"
+#include "SNodePanel.h"
+#include "ScopedTransaction.h"
+#include "Widgets/Docking/SDockTab.h"
+
+#define LOCTEXT_NAMESPACE "DataLinkGraphEditorToolkit"
+
+const FLazyName FDataLinkGraphEditorTool::GraphEditorTabID = TEXT("DataLinkGraphAssetToolkit_Graph");
+
+FDataLinkGraphEditorTool::FDataLinkGraphEditorTool(UDataLinkGraphAssetEditor* InAssetEditor)
+	: AssetEditor(InAssetEditor)
+	, GraphEditorCommands(MakeShared<FUICommandList>())
+{
+}
+
+void FDataLinkGraphEditorTool::BindCommands(const TSharedRef<FUICommandList>& InCommandList)
+{
+	GraphEditorCommands->Append(InCommandList);
+
+	const FGenericCommands& GenericCommands = FGenericCommands::Get();
+
+	GraphEditorCommands->MapAction(GenericCommands.SelectAll
+		, FExecuteAction::CreateSP(this, &FDataLinkGraphEditorTool::SelectAllNodes)
+	);
+
+	GraphEditorCommands->MapAction(GenericCommands.Delete
+		, FExecuteAction::CreateSP(this, &FDataLinkGraphEditorTool::DeleteSelectedNodes)
+		, FCanExecuteAction::CreateSP(this, &FDataLinkGraphEditorTool::CanDeleteSelectedNodes)
+	);
+
+	GraphEditorCommands->MapAction(GenericCommands.Copy
+		, FExecuteAction::CreateSP(this, &FDataLinkGraphEditorTool::CopySelectedNodes)
+		, FCanExecuteAction::CreateSP(this, &FDataLinkGraphEditorTool::CanCopySelectedNodes)
+	);
+
+	GraphEditorCommands->MapAction(GenericCommands.Cut
+		, FExecuteAction::CreateSP(this, &FDataLinkGraphEditorTool::CutSelectedNodes)
+		, FCanExecuteAction::CreateSP(this, &FDataLinkGraphEditorTool::CanCutSelectedNodes)
+	);
+
+	GraphEditorCommands->MapAction(GenericCommands.Paste
+		, FExecuteAction::CreateSP(this, &FDataLinkGraphEditorTool::PasteNodes)
+		, FCanExecuteAction::CreateSP(this, &FDataLinkGraphEditorTool::CanPasteNodes)
+	);
+
+	GraphEditorCommands->MapAction(GenericCommands.Duplicate
+		, FExecuteAction::CreateSP(this, &FDataLinkGraphEditorTool::DuplicateSelectedNodes)
+		, FCanExecuteAction::CreateSP(this, &FDataLinkGraphEditorTool::CanCopySelectedNodes)
+	);
+}
+
+void FDataLinkGraphEditorTool::RegisterTabSpawners(const TSharedRef<FTabManager>& InTabManager, const TSharedPtr<FWorkspaceItem>& InAssetEditorTabsCategory)
+{
+	InTabManager->RegisterTabSpawner(GraphEditorTabID, FOnSpawnTab::CreateSP(this, &FDataLinkGraphEditorTool::SpawnTab))
+		.SetDisplayName(LOCTEXT("Graph", "Graph"))
+		.SetGroup(InAssetEditorTabsCategory.ToSharedRef())
+		.SetIcon(FSlateIcon(FAppStyle::GetAppStyleSetName(), "GraphEditor.EventGraph_16x"));
+}
+
+void FDataLinkGraphEditorTool::UnregisterTabSpawners(const TSharedRef<FTabManager>& InTabManager)
+{
+	InTabManager->UnregisterTabSpawner(GraphEditorTabID);
+}
+
+void FDataLinkGraphEditorTool::CreateWidgets()
+{
+	FGraphAppearanceInfo AppearanceInfo;
+	AppearanceInfo.CornerText = LOCTEXT("GraphCornerText", "DATA LINK");
+
+	SGraphEditor::FGraphEditorEvents Events;
+	Events.OnSelectionChanged = SGraphEditor::FOnSelectionChanged::CreateSP(this, &FDataLinkGraphEditorTool::OnSelectedNodesChanged);
+
+	GraphEditor = SNew(SGraphEditor)
+		.AdditionalCommands(GraphEditorCommands)
+		.IsEditable(true)
+		.Appearance(AppearanceInfo)
+		.GraphToEdit(AssetEditor->GetDataLinkEdGraph())
+		.GraphEvents(Events)
+		.AutoExpandActionMenu(true)
+		.ShowGraphStateOverlay(false);
+}
+
+void FDataLinkGraphEditorTool::SelectAllNodes()
+{
+	if (GraphEditor.IsValid())
+	{
+		GraphEditor->SelectAllNodes();
+	}
+}
+
+bool FDataLinkGraphEditorTool::CanDeleteSelectedNodes() const
+{
+	if (!GraphEditor.IsValid())
+	{
+		return false;
+	}
+
+	int32 DeletableNodes = 0;
+
+	for (const UObject* SelectedNode : GraphEditor->GetSelectedNodes())
+	{
+		const UEdGraphNode* Node = Cast<UEdGraphNode>(SelectedNode);
+		if (Node && Node->CanUserDeleteNode())
+		{
+			++DeletableNodes;
+		}
+	}
+
+	return DeletableNodes > 0;
+}
+
+void FDataLinkGraphEditorTool::DeleteSelectedNodes()
+{
+	if (!GraphEditor.IsValid())
+	{
+		return;
+	}
+
+	UEdGraph* EdGraph = AssetEditor->GetDataLinkEdGraph();
+	if (!EdGraph)
+	{
+		return;
+	}
+
+	// Gather nodes to delete
+	TArray<UEdGraphNode*> NodesToDelete;
+	{
+		const TSet<UObject*>& SelectedNodes = GraphEditor->GetSelectedNodes();
+		NodesToDelete.Reserve(SelectedNodes.Num());
+
+		for (UObject* SelectedNode : GraphEditor->GetSelectedNodes())
+		{
+			UEdGraphNode* Node = Cast<UEdGraphNode>(SelectedNode);
+			if (Node && Node->CanUserDeleteNode())
+			{
+				NodesToDelete.Add(Node);
+			}
+		}
+	}
+
+	if (NodesToDelete.IsEmpty())
+	{
+		return;
+	}
+
+	FScopedTransaction Transaction(LOCTEXT("DeleteSelectedNodes", "Delete Selected Nodes"));
+	EdGraph->Modify();
+
+	for (UEdGraphNode* Node : NodesToDelete)
+	{
+		EdGraph->RemoveNode(Node);
+	}
+}
+
+bool FDataLinkGraphEditorTool::CanCopySelectedNodes() const
+{
+	return CanDuplicateSelectedNodes();
+}
+
+void FDataLinkGraphEditorTool::CopySelectedNodes()
+{
+	CopySelectedNodesInternal();
+}
+
+bool FDataLinkGraphEditorTool::CanCutSelectedNodes() const
+{
+	return CanCopySelectedNodes() && CanDeleteSelectedNodes();
+}
+
+void FDataLinkGraphEditorTool::CutSelectedNodes()
+{
+	UEdGraph* EdGraph = AssetEditor->GetDataLinkEdGraph();
+	if (!EdGraph)
+	{
+		return;
+	}
+
+	TSet<UObject*> CopiedNodes = CopySelectedNodesInternal();
+	if (CopiedNodes.IsEmpty())
+	{
+		return;
+	}
+
+	FScopedTransaction Transaction(LOCTEXT("CutSelectedNodes", "Cut Selected Nodes"));
+	EdGraph->Modify();
+
+	for (UObject* Node : CopiedNodes)
+	{
+		EdGraph->RemoveNode(CastChecked<UEdGraphNode>(Node));
+	}
+}
+
+bool FDataLinkGraphEditorTool::CanPasteNodes() const
+{
+	if (!GraphEditor.IsValid())
+	{
+		return false;
+	}
+
+	UEdGraph* EdGraph = AssetEditor->GetDataLinkEdGraph();
+	if (!EdGraph)
+	{
+		return false;
+	}
+
+	FString ClipboardContent;
+	FPlatformApplicationMisc::ClipboardPaste(ClipboardContent);
+	return FEdGraphUtilities::CanImportNodesFromText(EdGraph, ClipboardContent);
+}
+
+void FDataLinkGraphEditorTool::PasteNodes()
+{
+	if (!GraphEditor.IsValid())
+	{
+		return;
+	}
+
+	UEdGraph* EdGraph = AssetEditor->GetDataLinkEdGraph();
+	if (!EdGraph)
+	{
+		return;
+	}
+
+	const FScopedTransaction Transaction(LOCTEXT("PasteNodes", "Paste Nodes"));
+	EdGraph->Modify();
+
+	// Clear the selection set (newly pasted stuff will be selected)
+	GraphEditor->ClearSelectionSet();
+
+	// Grab the text to paste from the clipboard.
+	FString TextToImport;
+	FPlatformApplicationMisc::ClipboardPaste(TextToImport);
+
+	// Import the nodes
+	TSet<UEdGraphNode*> PastedNodes;
+	FEdGraphUtilities::ImportNodesFromText(EdGraph, TextToImport, /*out*/ PastedNodes);
+
+	if (PastedNodes.IsEmpty())
+	{
+		return;
+	}
+
+	// Average position of nodes so we can move them while still maintaining relative distances to each other
+	FVector2f AverageNodePosition = FVector2f::ZeroVector;
+	{
+		for (UEdGraphNode* PastedNode : PastedNodes)
+		{
+			AverageNodePosition.X += PastedNode->NodePosX;
+			AverageNodePosition.Y += PastedNode->NodePosY;
+		}
+		AverageNodePosition /= PastedNodes.Num();
+	}
+
+	const FVector2f PasteLocation = GraphEditor->GetPasteLocation2f();
+	const uint32 SnapGridSize = SNodePanel::GetSnapGridSize();
+
+	for (UEdGraphNode* PastedNode : PastedNodes)
+	{
+		// Select the newly pasted stuff
+		GraphEditor->SetNodeSelection(PastedNode, /*bSelect*/true);
+
+		PastedNode->NodePosX = (PastedNode->NodePosX - AverageNodePosition.X) + PasteLocation.X;
+		PastedNode->NodePosY = (PastedNode->NodePosY - AverageNodePosition.Y) + PasteLocation.Y;
+
+		PastedNode->SnapToGrid(SnapGridSize);
+
+		// Give new node a different Guid from the old one
+		PastedNode->CreateNewGuid();
+	}
+
+	GraphEditor->NotifyGraphChanged();
+}
+
+bool FDataLinkGraphEditorTool::CanDuplicateSelectedNodes() const
+{
+	if (!GraphEditor.IsValid())
+	{
+		return false;
+	}
+
+	int32 CanDuplicateCount = 0;
+
+	for (const UObject* SelectedNode : GraphEditor->GetSelectedNodes())
+	{
+		const UEdGraphNode* Node = Cast<UEdGraphNode>(SelectedNode);
+		if (Node && Node->CanDuplicateNode())
+		{
+			++CanDuplicateCount;
+		}
+	}
+
+	return CanDuplicateCount > 0;
+}
+
+void FDataLinkGraphEditorTool::DuplicateSelectedNodes()
+{
+	CopySelectedNodes();
+	PasteNodes();
+}
+
+TSet<UObject*> FDataLinkGraphEditorTool::CopySelectedNodesInternal()
+{
+	if (!GraphEditor.IsValid())
+	{
+		return {};
+	}
+
+	// Gather selected nodes
+	TSet<UObject*> NodesToCopy = GraphEditor->GetSelectedNodes();
+
+	// Notify nodes of copying, removing those that can't copy from the set
+	for (TSet<UObject*>::TIterator Iter(NodesToCopy); Iter; ++Iter)
+	{
+		UEdGraphNode* Node = Cast<UEdGraphNode>(*Iter);
+		if (Node && Node->CanDuplicateNode())
+		{
+			Node->PrepareForCopying();
+		}
+		else
+		{
+			Iter.RemoveCurrent();
+		}
+	}
+
+	FString ExportedText;
+	FEdGraphUtilities::ExportNodesToText(NodesToCopy, ExportedText);
+	FPlatformApplicationMisc::ClipboardCopy(*ExportedText);
+
+	return NodesToCopy;
+}
+
+TSharedRef<SDockTab> FDataLinkGraphEditorTool::SpawnTab(const FSpawnTabArgs& InTabArgs) const
+{
+	return SNew(SDockTab)
+		.Label(LOCTEXT("GraphTitle", "Graph"))
+		[
+			GraphEditor.ToSharedRef()
+		];
+}
+
+void FDataLinkGraphEditorTool::OnSelectedNodesChanged(const TSet<UObject*>& InSelectionSet)
+{
+	if (TSharedPtr<FDataLinkGraphAssetToolkit> Toolkit = AssetEditor->GetToolkit())
+	{
+		if (TSharedPtr<IDetailsView> DetailsView = Toolkit->GetDetailsView())
+		{
+			DetailsView->SetObjects(InSelectionSet.Array());
+		}
+	}
+}
+
+#undef LOCTEXT_NAMESPACE
diff --git a/Engine/Plugins/Experimental/DataLink/Source/DataLinkEditor/Private/Graph/DataLinkGraphEditorTool.h b/Engine/Plugins/Experimental/DataLink/Source/DataLinkEditor/Private/Graph/DataLinkGraphEditorTool.h
new file mode 100644
--- /dev/null
+++ b/Engine/Plugins/Experimental/DataLink/Source/DataLinkEditor/Private/Graph/DataLinkGraphEditorTool.h
@@ -0,0 +1,63 @@
+// Copyright Epic Games, Inc. All Rights Reserved.
+
+#pragma once
+
+#include "Containers/ContainersFwd.h"
+#include "Templates/SharedPointer.h"
+#include "UObject/NameTypes.h"
+#include "UObject/ObjectPtr.h"
+
+class FSpawnTabArgs;
+class FTabManager;
+class FUICommandList;
+class FWorkspaceItem;
+class SDockTab;
+class SGraphEditor;
+class UDataLinkGraphAssetEditor;
+class UObject;
+
+class FDataLinkGraphEditorTool : public TSharedFromThis<FDataLinkGraphEditorTool>
+{
+public:
+	static const FLazyName GraphEditorTabID;
+
+	explicit FDataLinkGraphEditorTool(UDataLinkGraphAssetEditor* InAssetEditor);
+
+	void BindCommands(const TSharedRef<FUICommandList>& InCommandList);
+
+	void RegisterTabSpawners(const TSharedRef<FTabManager>& InTabManager, const TSharedPtr<FWorkspaceItem>& InAssetEditorTabsCategory);
+	void UnregisterTabSpawners(const TSharedRef<FTabManager>& InTabManager);
+
+	void CreateWidgets();
+
+protected:
+	void SelectAllNodes();
+
+	bool CanDeleteSelectedNodes() const;
+	void DeleteSelectedNodes();
+
+	bool CanCopySelectedNodes() const;
+	void CopySelectedNodes();
+
+	bool CanCutSelectedNodes() const;
+	void CutSelectedNodes();
+
+	bool CanPasteNodes() const;
+	void PasteNodes();
+
+	bool CanDuplicateSelectedNodes() const;
+	void DuplicateSelectedNodes();
+
+private:
+	TSet<UObject*> CopySelectedNodesInternal();
+
+	TSharedRef<SDockTab> SpawnTab(const FSpawnTabArgs& InTabArgs) const;
+
+	void OnSelectedNodesChanged(const TSet<UObject*>& InSelectionSet);
+
+	TObjectPtr<UDataLinkGraphAssetEditor> AssetEditor;
+
+	TSharedRef<FUICommandList> GraphEditorCommands;
+
+	TSharedPtr<SGraphEditor> GraphEditor;
+};
diff --git a/Engine/Plugins/Experimental/DataLink/Source/DataLinkEditor/Private/Preview/DataLinkPreviewData.h b/Engine/Plugins/Experimental/DataLink/Source/DataLinkEditor/Private/Preview/DataLinkPreviewData.h
new file mode 100644
--- /dev/null
+++ b/Engine/Plugins/Experimental/DataLink/Source/DataLinkEditor/Private/Preview/DataLinkPreviewData.h
@@ -0,0 +1,21 @@
+// Copyright Epic Games, Inc. All Rights Reserved.
+
+#pragma once
+
+#include "DataLinkInstance.h"
+#include "StructUtils/InstancedStruct.h"
+#include "UObject/Object.h"
+#include "DataLinkPreviewData.generated.h"
+
+UCLASS()
+class UDataLinkPreviewData : public UObject
+{
+	GENERATED_BODY()
+
+public:
+	UPROPERTY(EditAnywhere, Category="Input Data")
+	FDataLinkInstance DataLinkInstance;
+
+	UPROPERTY(VisibleAnywhere, Category="Output Data", meta=(StructTypeConst))
+	FInstancedStruct OutputData;
+};
diff --git a/Engine/Plugins/Experimental/DataLink/Source/DataLinkEditor/Private/Preview/DataLinkPreviewTool.cpp b/Engine/Plugins/Experimental/DataLink/Source/DataLinkEditor/Private/Preview/DataLinkPreviewTool.cpp
new file mode 100644
--- /dev/null
+++ b/Engine/Plugins/Experimental/DataLink/Source/DataLinkEditor/Private/Preview/DataLinkPreviewTool.cpp
@@ -0,0 +1,235 @@
+// Copyright Epic Games, Inc. All Rights Reserved.
+
+#include "DataLinkPreviewTool.h"
+#include "DataLinkEditorLog.h"
+#include "DataLinkEditorNames.h"
+#include "DataLinkEditorStyle.h"
+#include "DataLinkEnums.h"
+#include "DataLinkExecutor.h"
+#include "DataLinkExecutorArguments.h"
+#include "DataLinkGraphAssetEditor.h"
+#include "DataLinkGraphAssetToolkit.h"
+#include "DataLinkGraphCommands.h"
+#include "DataLinkGraphEditorMenuContext.h"
+#include "DataLinkInstance.h"
+#include "DataLinkPreviewData.h"
+#include "DetailsView/DataLinkInstanceCustomization.h"
+#include "Modules/ModuleManager.h"
+#include "PropertyEditorModule.h"
+#include "ToolMenu.h"
+#include "ToolMenuSection.h"
+#include "ToolMenus.h"
+#include "Widgets/Docking/SDockTab.h"
+
+#define LOCTEXT_NAMESPACE "DataLinkPreviewTool"
+
+const FLazyName FDataLinkPreviewTool::PreviewTabID = TEXT("DataLinkGraphAssetToolkit_Preview");
+
+FDataLinkPreviewTool::FDataLinkPreviewTool(UDataLinkGraphAssetEditor* InAssetEditor)
+	: AssetEditor(InAssetEditor)
+{
+}
+
+void FDataLinkPreviewTool::Initialize()
+{
+	PreviewData = NewObject<UDataLinkPreviewData>(GetTransientPackage());
+	PreviewData->DataLinkInstance.DataLinkGraph = AssetEditor->GetDataLinkGraph();
+}
+
+void FDataLinkPreviewTool::BindCommands(const TSharedRef<FUICommandList>& InCommandList)
+{
+	const FDataLinkGraphCommands& GraphCommands = FDataLinkGraphCommands::Get();
+
+	InCommandList->MapAction(GraphCommands.Preview
+		, FExecuteAction::CreateSP(this, &FDataLinkPreviewTool::RunPreview)
+		, FCanExecuteAction::CreateSP(this, &FDataLinkPreviewTool::CanRunPreview));
+
+	InCommandList->MapAction(GraphCommands.ClearPreviewOutput
+		, FExecuteAction::CreateSP(this, &FDataLinkPreviewTool::ClearOutput)
+		, FCanExecuteAction::CreateSP(this, &FDataLinkPreviewTool::CanClearOutput));
+
+	InCommandList->MapAction(GraphCommands.ClearPreviewCache
+		, FExecuteAction::CreateSP(this, &FDataLinkPreviewTool::ClearCache)
+		, FCanExecuteAction::CreateSP(this, &FDataLinkPreviewTool::CanClearCache));
+}
+
+void FDataLinkPreviewTool::RegisterTabSpawners(const TSharedRef<FTabManager>& InTabManager, const TSharedPtr<FWorkspaceItem>& InAssetEditorTabsCategory)
+{
+	InTabManager->RegisterTabSpawner(PreviewTabID, FOnSpawnTab::CreateSP(this, &FDataLinkPreviewTool::SpawnTab))
+		.SetDisplayName(LOCTEXT("Preview", "Preview"))
+		.SetGroup(InAssetEditorTabsCategory.ToSharedRef())
+		.SetIcon(FSlateIcon(FDataLinkEditorStyle::Get().GetStyleSetName(), "DataLinkGraph.Preview"));
+}
+
+void FDataLinkPreviewTool::UnregisterTabSpawners(const TSharedRef<FTabManager>& InTabManager)
+{
+	InTabManager->UnregisterTabSpawner(PreviewTabID);
+}
+
+void FDataLinkPreviewTool::CreateWidgets()
+{
+	Initialize();
+	CreateDetailsView();
+	RegisterToolbar();
+}
+
+FString FDataLinkPreviewTool::GetReferencerName() const
+{
+	return TEXT("DataLinkPreviewTool");
+}
+
+void FDataLinkPreviewTool::AddReferencedObjects(FReferenceCollector& InCollector)
+{
+	InCollector.AddReferencedObject(PreviewData);
+}
+
+bool FDataLinkPreviewTool::CanRunPreview() const
+{
+	// Allow execution if there's not already one taking place
+	return !Executor.IsValid();
+}
+
+void FDataLinkPreviewTool::RunPreview()
+{
+	if (Executor.IsValid())
+	{
+		UE_LOG(LogDataLinkEditor, Error, TEXT("[%s] Data Link execution is in progress!"), Executor->GetContextName().GetData());
+		return;
+	}
+
+	if (!Sink.IsValid())
+	{
+		Sink = MakeShared<FDataLinkSink>();
+	}
+
+	Executor = FDataLinkExecutor::Create(FDataLinkExecutorArguments(PreviewData->DataLinkInstance)
+#if WITH_DATALINK_CONTEXT
+		.SetContextName(TEXT("Graph Editor Preview"))
+#endif
+		.SetContextObject(AssetEditor)
+		.SetSink(Sink)
+		.SetOnFinish(FOnDataLinkExecutionFinished::CreateSP(this, &FDataLinkPreviewTool::OnPreviewFinished)));
+
+	Executor->Run();
+}
+
+bool FDataLinkPreviewTool::CanClearOutput() const
+{
+	return PreviewData && PreviewData->OutputData.IsValid();
+}
+
+void FDataLinkPreviewTool::ClearOutput()
+{
+	if (PreviewData)
+	{
+		PreviewData->OutputData.Reset();
+	}
+}
+
+bool FDataLinkPreviewTool::CanClearCache() const
+{
+	// Even though Executor saves a shared ref of the Sink and will not affect if this Sink ref is cleared during execution,
+	// it might set wrong expectations to the user... e.g. user might think that clearing cache during execution might have an effect.
+	// So only allow clearing cache if there's a valid sink and no execution is taking place.
+	return Sink.IsValid() && !Executor.IsValid();
+}
+
+void FDataLinkPreviewTool::ClearCache()
+{
+	Sink.Reset();
+}
+
+void FDataLinkPreviewTool::OnPreviewFinished(const FDataLinkExecutor& InExecutor, EDataLinkExecutionResult InResult, FConstStructView InOutputDataView)
+{
+	Executor.Reset();
+
+	const UScriptStruct* OutputDataStruct = InOutputDataView.GetScriptStruct();
+	if (!OutputDataStruct || InResult == EDataLinkExecutionResult::Failed)
+	{
+		PreviewData->OutputData.Reset();
+		return;
+	}
+
+	PreviewData->OutputData.InitializeAs(OutputDataStruct);
+	OutputDataStruct->CopyScriptStruct(PreviewData->OutputData.GetMutableMemory(), InOutputDataView.GetMemory());
+	PreviewDetails->ForceRefresh();
+}
+
+void FDataLinkPreviewTool::RegisterToolbar()
+{
+	UToolMenus* ToolMenus = UToolMenus::Get();
+	if (!ToolMenus || ToolMenus->IsMenuRegistered(UE::DataLinkEditor::PreviewToolbarName))
+	{
+		return;
+	}
+
+	UToolMenu* const ToolbarMenu = ToolMenus->RegisterMenu(UE::DataLinkEditor::PreviewToolbarName, NAME_None, EMultiBoxType::SlimHorizontalToolBar);
+	if (!ToolbarMenu)
+	{
+		return;
+	}
+
+	const FDataLinkGraphCommands& GraphCommands = FDataLinkGraphCommands::Get();
+
+	FToolMenuSection& PreviewSection = ToolbarMenu->FindOrAddSection(UE::DataLinkEditor::PreviewSectionName);
+	PreviewSection.AddEntry(FToolMenuEntry::InitToolBarButton(GraphCommands.Preview));
+	PreviewSection.AddEntry(FToolMenuEntry::InitToolBarButton(GraphCommands.ClearPreviewOutput));
+	PreviewSection.AddEntry(FToolMenuEntry::InitToolBarButton(GraphCommands.ClearPreviewCache));
+}
+
+void FDataLinkPreviewTool::CreateDetailsView()
+{
+	FPropertyEditorModule& PropertyEditorModule = FModuleManager::GetModuleChecked<FPropertyEditorModule>("PropertyEditor");
+
+	FDetailsViewArgs DetailsViewArgs;
+	DetailsViewArgs.DefaultsOnlyVisibility = EEditDefaultsOnlyNodeVisibility::Automatic;
+	DetailsViewArgs.NameAreaSettings = FDetailsViewArgs::HideNameArea;
+	DetailsViewArgs.bHideSelectionTip = true;
+
+	PreviewDetails = PropertyEditorModule.CreateDetailView(DetailsViewArgs);
+
+	PreviewDetails->RegisterInstancedCustomPropertyTypeLayout(FDataLinkInstance::StaticStruct()->GetFName()
+		, FOnGetPropertyTypeCustomizationInstance::CreateStatic(&FDataLinkInstanceCustomization::MakeInstance, /*bGenerateHeader*/false));
+
+	PreviewDetails->SetObject(PreviewData);
+}
+
+TSharedRef<SWidget> FDataLinkPreviewTool::CreateContentWidget() const
+{
+	return SNew(SVerticalBox)
+		+ SVerticalBox::Slot()
+		.AutoHeight()
+		[
+			CreateToolbar()
+		]
+		+SVerticalBox::Slot()
+		.FillHeight(1.f)
+		[
+			PreviewDetails.ToSharedRef()
+		];
+}
+
+TSharedRef<SWidget> FDataLinkPreviewTool::CreateToolbar() const
+{
+	UToolMenus* const ToolMenus = UToolMenus::Get();
+	check(AssetEditor && ToolMenus);
+
+	FToolMenuContext Context(AssetEditor->GetToolkitCommands());
+
+	UDataLinkGraphEditorMenuContext* MenuContext = NewObject<UDataLinkGraphEditorMenuContext>();
+	MenuContext->ToolkitWeak = AssetEditor->GetToolkit();
+	Context.AddObject(MenuContext);
+
+	return ToolMenus->GenerateWidget(UE::DataLinkEditor::PreviewToolbarName, Context);
+}
+
+TSharedRef<SDockTab> FDataLinkPreviewTool::SpawnTab(const FSpawnTabArgs& InTabArgs) const
+{
+	return SNew(SDockTab)
+		.Label(LOCTEXT("PreviewTitle", "Preview"))
+		[
+			CreateContentWidget()
+		];
+}
+
+#undef LOCTEXT_NAMESPACE
diff --git a/Engine/Plugins/Experimental/DataLink/Source/DataLinkEditor/Private/Preview/DataLinkPreviewTool.h b/Engine/Plugins/Experimental/DataLink/Source/DataLinkEditor/Private/Preview/DataLinkPreviewTool.h
new file mode 100644
--- /dev/null
+++ b/Engine/Plugins/Experimental/DataLink/Source/DataLinkEditor/Private/Preview/DataLinkPreviewTool.h
@@ -0,0 +1,85 @@
+// Copyright Epic Games, Inc. All Rights Reserved.
+
+#pragma once
+
+#include "Templates/SharedPointer.h"
+#include "UObject/GCObject.h"
+
+class FDataLinkExecutor;
+class FSpawnTabArgs;
+class FTabManager;
+class FUICommandList;
+class FWorkspaceItem;
+class IDetailsView;
+class SDockTab;
+class SWidget;
+class UDataLinkGraphAssetEditor;
+class UDataLinkPreviewData;
+enum class EDataLinkExecutionResult : uint8;
+struct FConstStructView;
+struct FDataLinkSink;
+
+/** Handles executing a preview of the Data Link Graph with input data from within the Editor itself */
+class FDataLinkPreviewTool : public FGCObject, public TSharedFromThis<FDataLinkPreviewTool>
+{
+public:
+	static const FLazyName PreviewTabID;
+
+	explicit FDataLinkPreviewTool(UDataLinkGraphAssetEditor* InAssetEditor);
+
+	void Initialize();
+
+	void BindCommands(const TSharedRef<FUICommandList>& InCommandList);
+
+	void RegisterTabSpawners(const TSharedRef<FTabManager>& InTabManager, const TSharedPtr<FWorkspaceItem>& InAssetEditorTabsCategory);
+
+	void UnregisterTabSpawners(const TSharedRef<FTabManager>& InTabManager);
+
+	void CreateWidgets();
+
+	const UDataLinkPreviewData* GetPreviewData() const
+	{
+		return PreviewData;
+	}
+
+protected:
+	//~ Begin FGCObject
+	virtual FString GetReferencerName() const override;
+	virtual void AddReferencedObjects(FReferenceCollector& InCollector) override;
+	//~ End FGCObject
+
+private:
+	bool CanRunPreview() const;
+	void RunPreview();
+
+	bool CanClearOutput() const;
+	void ClearOutput();
+
+	bool CanClearCache() const;
+	void ClearCache();
+
+	void OnPreviewFinished(const FDataLinkExecutor& InExecutor, EDataLinkExecutionResult InResult, FConstStructView InOutputDataView);
+
+	void RegisterToolbar();
+
+	void CreateDetailsView();
+
+	TSharedRef<SWidget> CreateContentWidget() const;
+
+	TSharedRef<SWidget> CreateToolbar() const;
+
+	TSharedRef<SDockTab> SpawnTab(const FSpawnTabArgs& InTabArgs) const;
+
+	TObjectPtr<UDataLinkGraphAssetEditor> AssetEditor;
+
+	TSharedPtr<IDetailsView> PreviewDetails;
+
+	/** Preview object holding the input and output data */
+	TObjectPtr<UDataLinkPreviewData> PreviewData;
+
+	/** Current Executor taking place */
+	TSharedPtr<FDataLinkExecutor> Executor;
+
+	/** Data Sink to use for Executor */
+	TSharedPtr<FDataLinkSink> Sink;
+};
diff --git a/Engine/Plugins/Experimental/DataLink/Source/DataLinkEditor/Public/DataLinkEditorNames.h b/Engine/Plugins/Experimental/DataLink/Source/DataLinkEditor/Public/DataLinkEditorNames.h
new file mode 100644
--- /dev/null
+++ b/Engine/Plugins/Experimental/DataLink/Source/DataLinkEditor/Public/DataLinkEditorNames.h
@@ -0,0 +1,11 @@
+// Copyright Epic Games, Inc. All Rights Reserved.
+
+#pragma once
+
+#include "UObject/NameTypes.h"
+
+namespace UE::DataLinkEditor
+{
+	const FLazyName PreviewToolbarName(TEXT("DataLinkGraphAssetToolkit_PreviewToolbar"));
+	const FLazyName PreviewSectionName(TEXT("Preview"));
+}
diff --git a/Engine/Plugins/Experimental/DataLink/Source/DataLinkEditor/Public/IDataLinkEditorMenuContext.h b/Engine/Plugins/Experimental/DataLink/Source/DataLinkEditor/Public/IDataLinkEditorMenuContext.h
new file mode 100644
--- /dev/null
+++ b/Engine/Plugins/Experimental/DataLink/Source/DataLinkEditor/Public/IDataLinkEditorMenuContext.h
@@ -0,0 +1,29 @@
+// Copyright Epic Games, Inc. All Rights Reserved.
+
+#pragma once
+
+#include "UObject/Interface.h"
+#include "IDataLinkEditorMenuContext.generated.h"
+
+struct FConstStructView;
+
+UINTERFACE(MinimalAPI, meta=(CannotImplementInterfaceInBlueprint))
+class UDataLinkEditorMenuContext : public UInterface
+{
+	GENERATED_BODY()
+};
+
+class IDataLinkEditorMenuContext
+{
+	GENERATED_BODY()
+
+public:
+	/** 
+	 * Finds the latest preview output data. Could return invalid struct.
+	 * @note Existing views can be invalidated at any time by the user so handle these views as short-lived!
+	 */
+	virtual FConstStructView FindPreviewOutputData() const = 0;
+
+	/** Retrieves the asset path of the currently edited asset */
+	virtual FString GetAssetPath() const = 0;
+};
diff --git a/Engine/Plugins/Experimental/DataLink/Source/DataLinkHttp/DataLinkHttp.Build.cs b/Engine/Plugins/Experimental/DataLink/Source/DataLinkHttp/DataLinkHttp.Build.cs
new file mode 100644
--- /dev/null
+++ b/Engine/Plugins/Experimental/DataLink/Source/DataLinkHttp/DataLinkHttp.Build.cs
@@ -0,0 +1,27 @@
+// Copyright Epic Games, Inc. All Rights Reserved.
+
+using UnrealBuildTool;
+
+public class DataLinkHttp : ModuleRules
+{
+    public DataLinkHttp(ReadOnlyTargetRules Target) : base(Target)
+    {
+        PCHUsage = ModuleRules.PCHUsageMode.UseExplicitOrSharedPCHs;
+
+        PublicDependencyModuleNames.AddRange(
+            new string[]
+            {
+                "Core",
+                "CoreUObject",
+                "DataLink",
+            }
+        );
+
+        PrivateDependencyModuleNames.AddRange(
+            new string[]
+            {
+                "HTTP",
+            }
+        );
+    }
+}
\ No newline at end of file
diff --git a/Engine/Plugins/Experimental/DataLink/Source/DataLinkHttp/Private/DataLinkHttpModule.cpp b/Engine/Plugins/Experimental/DataLink/Source/DataLinkHttp/Private/DataLinkHttpModule.cpp
new file mode 100644
--- /dev/null
+++ b/Engine/Plugins/Experimental/DataLink/Source/DataLinkHttp/Private/DataLinkHttpModule.cpp
@@ -0,0 +1,5 @@
+// Copyright Epic Games, Inc. All Rights Reserved.
+
+#include "Modules/ModuleManager.h"
+
+IMPLEMENT_MODULE(FDefaultModuleImpl, DataLinkHttp)
diff --git a/Engine/Plugins/Experimental/DataLink/Source/DataLinkHttp/Private/DataLinkHttpSource.cpp b/Engine/Plugins/Experimental/DataLink/Source/DataLinkHttp/Private/DataLinkHttpSource.cpp
new file mode 100644
--- /dev/null
+++ b/Engine/Plugins/Experimental/DataLink/Source/DataLinkHttp/Private/DataLinkHttpSource.cpp
@@ -0,0 +1,86 @@
+// Copyright Epic Games, Inc. All Rights Reserved.
+
+#include "DataLinkHttpSource.h"
+#include "DataLinkCoreTypes.h"
+#include "DataLinkEnums.h"
+#include "DataLinkExecutor.h"
+#include "DataLinkHttpNames.h"
+#include "DataLinkHttpSettings.h"
+#include "DataLinkNames.h"
+#include "DataLinkPinBuilder.h"
+#include "DataLinkSink.h"
+#include "HttpModule.h"
+#include "Interfaces/IHttpResponse.h"
+
+#define LOCTEXT_NAMESPACE "DataLinkHttpSource"
+
+void UDataLinkHttpSource::OnBuildPins(FDataLinkPinBuilder& Inputs, FDataLinkPinBuilder& Outputs) const
+{
+	Super::OnBuildPins(Inputs, Outputs);
+
+	Inputs.Add(UE::DataLinkHttp::InputHttpSettings)
+		.SetDisplayName(LOCTEXT("HttpSettingsDisplay", "Http Settings"))
+		.SetStruct<FDataLinkHttpSettings>();
+
+	Outputs.Add(UE::DataLink::OutputDefault)
+		.SetDisplayName(LOCTEXT("OutputDisplay", "Response"))
+		.SetStruct<FDataLinkString>();
+}
+
+EDataLinkExecutionReply UDataLinkHttpSource::OnExecute(FDataLinkExecutor& InExecutor) const
+{
+	const FDataLinkNodeInstance& NodeInstance = InExecutor.GetNodeInstance(this);
+
+	const FDataLinkInputDataViewer& InputDataViewer = NodeInstance.GetInputDataViewer();
+
+	const FDataLinkHttpSettings& InputData = InputDataViewer.Get<FDataLinkHttpSettings>(UE::DataLinkHttp::InputHttpSettings);
+
+	const TSharedRef<IHttpRequest> HttpRequest = FHttpModule::Get().CreateRequest();
+	HttpRequest->SetURL(InputData.URL);
+	HttpRequest->SetVerb(InputData.Verb);
+
+	for (const TPair<FString, FString>& HeaderPair : InputData.Headers)
+	{
+		HttpRequest->SetHeader(HeaderPair.Key, HeaderPair.Value);
+	}
+
+	HttpRequest->SetContentAsString(InputData.Body);
+	HttpRequest->SetDelegateThreadPolicy(EHttpRequestDelegateThreadPolicy::CompleteOnGameThread);
+
+	HttpRequest->OnProcessRequestComplete().BindWeakLambda(this,
+		[this, ExecutorWeak = InExecutor.AsWeak()](FHttpRequestPtr InRequest, FHttpResponsePtr InResponse, bool bInProcessedSuccessfully)
+		{
+			// Delegate Thread Policy set to Complete on Game Thread
+			check(IsInGameThread());
+
+			TSharedPtr<FDataLinkExecutor> Executor = ExecutorWeak.Pin();
+			if (!Executor.IsValid())
+			{
+				return;
+			}
+
+			if (bInProcessedSuccessfully)
+			{
+				const FDataLinkNodeInstance& NodeInstance = Executor->GetNodeInstance(this);
+
+				const FDataLinkOutputDataViewer& OutputDataViewer = NodeInstance.GetOutputDataViewer();
+
+				FDataLinkString& OutputDataView = OutputDataViewer.Get<FDataLinkString>(UE::DataLink::OutputDefault);
+				OutputDataView.Value = InResponse->GetContentAsString();
+
+				Executor->SucceedNode(this, FConstStructView::Make(OutputDataView));
+			}
+			else
+			{
+				Executor->FailNode(this);
+			}
+		});
+
+	const bool bRequestProcessed = HttpRequest->ProcessRequest();
+
+	return bRequestProcessed
+		? EDataLinkExecutionReply::Handled
+		: EDataLinkExecutionReply::Unhandled;
+}
+
+#undef LOCTEXT_NAMESPACE
diff --git a/Engine/Plugins/Experimental/DataLink/Source/DataLinkHttp/Private/DataLinkNodeHttpSettingsBuilder.cpp b/Engine/Plugins/Experimental/DataLink/Source/DataLinkHttp/Private/DataLinkNodeHttpSettingsBuilder.cpp
new file mode 100644
--- /dev/null
+++ b/Engine/Plugins/Experimental/DataLink/Source/DataLinkHttp/Private/DataLinkNodeHttpSettingsBuilder.cpp
@@ -0,0 +1,58 @@
+// Copyright Epic Games, Inc. All Rights Reserved.
+
+#include "DataLinkNodeHttpSettingsBuilder.h"
+#include "DataLinkCoreTypes.h"
+#include "DataLinkExecutor.h"
+#include "DataLinkHttpSettings.h"
+#include "DataLinkNames.h"
+#include "DataLinkNodeInstance.h"
+#include "DataLinkPinBuilder.h"
+
+#define LOCTEXT_NAMESPACE "DataLinkHttpSettingsBuilder"
+
+void UDataLinkNodeHttpSettingsBuilder::OnBuildPins(FDataLinkPinBuilder& Inputs, FDataLinkPinBuilder& Outputs) const
+{
+	FDataLinkStringBuilder(URLSegments, Tokens)
+		.BuildInputPins(Inputs);
+
+	Outputs.Add(UE::DataLink::OutputDefault)
+		.SetDisplayName(LOCTEXT("HttpSettings", "Http Settings"))
+		.SetStruct<FDataLinkHttpSettings>();
+}
+
+EDataLinkExecutionReply UDataLinkNodeHttpSettingsBuilder::OnExecute(FDataLinkExecutor& InExecutor) const
+{
+	const FDataLinkNodeInstance& NodeInstance = InExecutor.GetNodeInstance(this);
+
+	const FDataLinkOutputDataViewer& OutputDataViewer = NodeInstance.GetOutputDataViewer();
+
+	FDataLinkHttpSettings& HttpSettings = OutputDataViewer.Get<FDataLinkHttpSettings>(UE::DataLink::OutputDefault);
+
+	const FDataLinkStringBuilder URLBuilder(URLSegments, Tokens);
+	if (!URLBuilder.BuildString(NodeInstance.GetInputDataViewer(), HttpSettings.URL))
+	{
+		InExecutor.FailNode(this);
+		return EDataLinkExecutionReply::Handled;
+	}
+
+	HttpSettings.Verb = Verb;
+	HttpSettings.Headers = Headers;
+	HttpSettings.Body = Body;
+
+	InExecutor.SucceedNode(this, FConstStructView::Make(HttpSettings));
+	return EDataLinkExecutionReply::Handled;
+}
+
+#if WITH_EDITOR
+void UDataLinkNodeHttpSettingsBuilder::PostEditChangeProperty(FPropertyChangedEvent& InPropertyChangedEvent)
+{
+	Super::PostEditChangeProperty(InPropertyChangedEvent);
+
+	if (InPropertyChangedEvent.GetMemberPropertyName() == GET_MEMBER_NAME_CHECKED(UDataLinkNodeHttpSettingsBuilder, URLSegments))
+	{
+		FDataLinkStringBuilder::GatherTokens(URLSegments, Tokens);
+	}
+}
+#endif
+
+#undef LOCTEXT_NAMESPACE
diff --git a/Engine/Plugins/Experimental/DataLink/Source/DataLinkHttp/Public/DataLinkHttpNames.h b/Engine/Plugins/Experimental/DataLink/Source/DataLinkHttp/Public/DataLinkHttpNames.h
new file mode 100644
--- /dev/null
+++ b/Engine/Plugins/Experimental/DataLink/Source/DataLinkHttp/Public/DataLinkHttpNames.h
@@ -0,0 +1,10 @@
+// Copyright Epic Games, Inc. All Rights Reserved.
+
+#pragma once
+
+#include "UObject/NameTypes.h"
+
+namespace UE::DataLinkHttp
+{
+	const FLazyName InputHttpSettings(TEXT("InputHttpSettings"));
+}
diff --git a/Engine/Plugins/Experimental/DataLink/Source/DataLinkHttp/Public/DataLinkHttpSettings.h b/Engine/Plugins/Experimental/DataLink/Source/DataLinkHttp/Public/DataLinkHttpSettings.h
new file mode 100644
--- /dev/null
+++ b/Engine/Plugins/Experimental/DataLink/Source/DataLinkHttp/Public/DataLinkHttpSettings.h
@@ -0,0 +1,25 @@
+// Copyright Epic Games, Inc. All Rights Reserved.
+
+#pragma once
+
+#include "Containers/Map.h"
+#include "Containers/UnrealString.h"
+#include "DataLinkHttpSettings.generated.h"
+
+USTRUCT(BlueprintType, DisplayName="Data Link Http Settings")
+struct FDataLinkHttpSettings
+{
+	GENERATED_BODY()
+
+	UPROPERTY(EditAnywhere, BlueprintReadWrite, Category="Data Link Http")
+	FString URL;
+
+	UPROPERTY(EditAnywhere, BlueprintReadWrite, Category="Data Link Http")
+	FString Verb = TEXT("GET");
+
+	UPROPERTY(EditAnywhere, BlueprintReadWrite, Category="Data Link Http")
+	TMap<FString, FString> Headers;
+
+	UPROPERTY(EditAnywhere, BlueprintReadWrite, Category="Data Link Http")
+	FString Body;
+};
diff --git a/Engine/Plugins/Experimental/DataLink/Source/DataLinkHttp/Public/DataLinkHttpSource.h b/Engine/Plugins/Experimental/DataLink/Source/DataLinkHttp/Public/DataLinkHttpSource.h
new file mode 100644
--- /dev/null
+++ b/Engine/Plugins/Experimental/DataLink/Source/DataLinkHttp/Public/DataLinkHttpSource.h
@@ -0,0 +1,18 @@
+// Copyright Epic Games, Inc. All Rights Reserved.
+
+#pragma once
+
+#include "DataLinkNode.h"
+#include "DataLinkHttpSource.generated.h"
+
+UCLASS(MinimalAPI, DisplayName="Http Request", Category="Http")
+class UDataLinkHttpSource : public UDataLinkNode
+{
+    GENERATED_BODY()
+
+protected:
+    //~ Begin UDataLinkNode
+    DATALINKHTTP_API virtual void OnBuildPins(FDataLinkPinBuilder& Inputs, FDataLinkPinBuilder& Outputs) const override;
+    DATALINKHTTP_API virtual EDataLinkExecutionReply OnExecute(FDataLinkExecutor& InExecutor) const override;
+    //~ End UDataLinkNode
+};
diff --git a/Engine/Plugins/Experimental/DataLink/Source/DataLinkHttp/Public/DataLinkNodeHttpSettingsBuilder.h b/Engine/Plugins/Experimental/DataLink/Source/DataLinkHttp/Public/DataLinkNodeHttpSettingsBuilder.h
new file mode 100644
--- /dev/null
+++ b/Engine/Plugins/Experimental/DataLink/Source/DataLinkHttp/Public/DataLinkNodeHttpSettingsBuilder.h
@@ -0,0 +1,47 @@
+// Copyright Epic Games, Inc. All Rights Reserved.
+
+#pragma once
+
+#include "Containers/Array.h"
+#include "Containers/Map.h"
+#include "Containers/UnrealString.h"
+#include "DataLinkNode.h"
+#include "Nodes/String/DataLinkStringBuilder.h"
+#include "DataLinkNodeHttpSettingsBuilder.generated.h"
+
+/** Builder Node to help layer Http Settings with a URL Builder Interface */
+UCLASS(MinimalAPI, DisplayName="Http Settings Builder", Category="Http")
+class UDataLinkNodeHttpSettingsBuilder : public UDataLinkNode
+{
+	GENERATED_BODY()
+
+protected:
+	//~ Begin UDataLinkNode
+	DATALINKHTTP_API virtual void OnBuildPins(FDataLinkPinBuilder& Inputs, FDataLinkPinBuilder& Outputs) const override;
+	DATALINKHTTP_API virtual EDataLinkExecutionReply OnExecute(FDataLinkExecutor& InExecutor) const override;
+	//~ End UDataLinkNode
+
+	//~ Begin UObject
+#if WITH_EDITOR
+	DATALINKHTTP_API virtual void PostEditChangeProperty(FPropertyChangedEvent& InPropertyChangedEvent) override;
+#endif
+	//~ End UObject
+
+private:
+	/** Segments to build the URL from, where a token can be defined as {Token Name} in its own isolated entry */
+	UPROPERTY(EditAnywhere, Category="URL", DisplayName="URL Segments", meta=(InvalidatesNode))
+	TArray<FString> URLSegments;
+
+	/** The tokens found within the Segment Array */
+	UPROPERTY()
+	TArray<FDataLinkStringBuilderToken> Tokens;
+
+	UPROPERTY(EditAnywhere, Category="Settings", AdvancedDisplay)
+	FString Verb = TEXT("GET");
+
+	UPROPERTY(EditAnywhere, Category="Settings", AdvancedDisplay)
+	TMap<FString, FString> Headers;
+
+	UPROPERTY(EditAnywhere, Category="Settings", AdvancedDisplay)
+	FString Body;
+};
diff --git a/Engine/Plugins/Experimental/DataLink/Source/DataLinkJson/DataLinkJson.Build.cs b/Engine/Plugins/Experimental/DataLink/Source/DataLinkJson/DataLinkJson.Build.cs
new file mode 100644
--- /dev/null
+++ b/Engine/Plugins/Experimental/DataLink/Source/DataLinkJson/DataLinkJson.Build.cs
@@ -0,0 +1,28 @@
+// Copyright Epic Games, Inc. All Rights Reserved.
+
+using UnrealBuildTool;
+
+public class DataLinkJson : ModuleRules
+{
+    public DataLinkJson(ReadOnlyTargetRules Target) : base(Target)
+    {
+        PCHUsage = ModuleRules.PCHUsageMode.UseExplicitOrSharedPCHs;
+
+        PublicDependencyModuleNames.AddRange(
+            new string[]
+            {
+                "Core",
+                "CoreUObject",
+                "DataLink",
+            }
+        );
+
+        PrivateDependencyModuleNames.AddRange(
+            new string[]
+            {
+                "Engine",
+                "JsonUtilities",
+            }
+        );
+    }
+}
\ No newline at end of file
diff --git a/Engine/Plugins/Experimental/DataLink/Source/DataLinkJson/Private/DataLinkJsonFromStruct.cpp b/Engine/Plugins/Experimental/DataLink/Source/DataLinkJson/Private/DataLinkJsonFromStruct.cpp
new file mode 100644
--- /dev/null
+++ b/Engine/Plugins/Experimental/DataLink/Source/DataLinkJson/Private/DataLinkJsonFromStruct.cpp
@@ -0,0 +1,59 @@
+// Copyright Epic Games, Inc. All Rights Reserved.
+
+#include "DataLinkJsonFromStruct.h"
+#include "DataLinkExecutor.h"
+#include "DataLinkJsonNames.h"
+#include "DataLinkNames.h"
+#include "DataLinkNodeInstance.h"
+#include "DataLinkPinBuilder.h"
+#include "JsonObjectConverter.h"
+#include "JsonObjectWrapper.h"
+
+#define LOCTEXT_NAMESPACE "DataLinkJsonFromStruct"
+
+void UDataLinkJsonFromStruct::OnBuildPins(FDataLinkPinBuilder& Inputs, FDataLinkPinBuilder& Outputs) const
+{
+	Super::OnBuildPins(Inputs, Outputs);
+
+	// No struct specified as this generically accepts any script struct and converts it to json object
+	Inputs.Add(UE::DataLinkJson::InputStruct)
+		.SetDisplayName(LOCTEXT("InputStructDisplay", "Struct"));
+
+	Outputs.Add(UE::DataLink::OutputDefault)
+		.SetDisplayName(LOCTEXT("InputJsonObjectDisplay", "Json"))
+		.SetStruct<FJsonObjectWrapper>();
+}
+
+EDataLinkExecutionReply UDataLinkJsonFromStruct::OnExecute(FDataLinkExecutor& InExecutor) const
+{
+	const FDataLinkNodeInstance& NodeInstance = InExecutor.GetNodeInstance(this);
+
+	FConstStructView InputStruct = NodeInstance.GetInputDataViewer().Find(UE::DataLinkJson::InputStruct);
+	if (!InputStruct.IsValid())
+	{
+		return EDataLinkExecutionReply::Unhandled;
+	}
+
+	FJsonObjectWrapper& Wrapper = NodeInstance.GetOutputDataViewer().Get<FJsonObjectWrapper>(UE::DataLink::OutputDefault);
+	if (!Wrapper.JsonObject.IsValid())
+	{
+		return EDataLinkExecutionReply::Unhandled;
+	}
+
+	const bool bResult = FJsonObjectConverter::UStructToJsonObject(InputStruct.GetScriptStruct()
+		, InputStruct.GetMemory()
+		, Wrapper.JsonObject.ToSharedRef());
+
+	if (!bResult)
+	{
+		return EDataLinkExecutionReply::Unhandled;
+	}
+
+	// Set the json object wrapper's string to the latest json object result
+	Wrapper.JsonObjectToString(Wrapper.JsonString);
+
+	InExecutor.SucceedNode(this, FConstStructView::Make(Wrapper));
+	return EDataLinkExecutionReply::Handled;
+}
+
+#undef LOCTEXT_NAMESPACE
diff --git a/Engine/Plugins/Experimental/DataLink/Source/DataLinkJson/Private/DataLinkJsonModule.cpp b/Engine/Plugins/Experimental/DataLink/Source/DataLinkJson/Private/DataLinkJsonModule.cpp
new file mode 100644
--- /dev/null
+++ b/Engine/Plugins/Experimental/DataLink/Source/DataLinkJson/Private/DataLinkJsonModule.cpp
@@ -0,0 +1,8 @@
+// Copyright Epic Games, Inc. All Rights Reserved.
+
+#include "DataLinkJsonLog.h"
+#include "Modules/ModuleManager.h"
+
+DEFINE_LOG_CATEGORY(LogDataLinkJson);
+
+IMPLEMENT_MODULE(FDefaultModuleImpl, DataLinkJson)
diff --git a/Engine/Plugins/Experimental/DataLink/Source/DataLinkJson/Private/DataLinkJsonStructSimpleMapping.cpp b/Engine/Plugins/Experimental/DataLink/Source/DataLinkJson/Private/DataLinkJsonStructSimpleMapping.cpp
new file mode 100644
--- /dev/null
+++ b/Engine/Plugins/Experimental/DataLink/Source/DataLinkJson/Private/DataLinkJsonStructSimpleMapping.cpp
@@ -0,0 +1,46 @@
+// Copyright Epic Games, Inc. All Rights Reserved.
+
+#include "DataLinkJsonStructSimpleMapping.h"
+#include "DataLinkJsonLog.h"
+#include "DataLinkJsonUtils.h"
+#include "JsonObjectConverter.h"
+#include "StructUtils/StructView.h"
+#include "UObject/PropertyAccessUtil.h"
+
+bool UDataLinkJsonStructSimpleMapping::Apply(const TSharedRef<FJsonObject>& InSourceJson, const FStructView& InTargetStructView) const
+{
+	const UScriptStruct* const TargetStruct = InTargetStructView.GetScriptStruct();
+
+	uint8* const TargetMemory = InTargetStructView.GetMemory();
+
+	for (const TPair<FString, FName>& FieldMapping : FieldMappings)
+	{
+		FProperty* const TargetProperty = PropertyAccessUtil::FindPropertyByName(FieldMapping.Value, TargetStruct);
+		if (!TargetProperty)
+		{
+			UE_LOG(LogDataLinkJson, Error, TEXT("Field '%s' not found in struct '%s'")
+				, *FieldMapping.Value.ToString()
+				, *TargetStruct->GetName());
+			continue;
+		}
+
+		const TSharedPtr<FJsonValue> JsonValue = UE::DataLinkJson::FindJsonValue(InSourceJson, FieldMapping.Key);
+		if (!JsonValue.IsValid())
+		{
+			UE_LOG(LogDataLinkJson, Error, TEXT("Field '%s' not found in json")
+				, *FieldMapping.Key);
+			continue;
+		}
+
+		uint8* const TargetPropertyValue = TargetProperty->ContainerPtrToValuePtr<uint8>(TargetMemory);
+		if (!FJsonObjectConverter::JsonValueToUProperty(JsonValue, TargetProperty, TargetPropertyValue))
+		{
+			UE_LOG(LogDataLinkJson, Error, TEXT("Could not copy Json Value with key '%s' to property '%s' in struct '%s'")
+				, *FieldMapping.Key
+				, *TargetProperty->GetName()
+				, *TargetStruct->GetName());
+		}
+	}
+
+	return true;
+}
diff --git a/Engine/Plugins/Experimental/DataLink/Source/DataLinkJson/Private/DataLinkJsonToStruct.cpp b/Engine/Plugins/Experimental/DataLink/Source/DataLinkJson/Private/DataLinkJsonToStruct.cpp
new file mode 100644
--- /dev/null
+++ b/Engine/Plugins/Experimental/DataLink/Source/DataLinkJson/Private/DataLinkJsonToStruct.cpp
@@ -0,0 +1,72 @@
+// Copyright Epic Games, Inc. All Rights Reserved.
+
+#include "DataLinkJsonToStruct.h"
+#include "DataLinkExecutor.h"
+#include "DataLinkJsonNames.h"
+#include "DataLinkJsonStructMapping.h"
+#include "DataLinkNames.h"
+#include "DataLinkNodeInstance.h"
+#include "DataLinkPinBuilder.h"
+#include "JsonObjectConverter.h"
+#include "JsonObjectWrapper.h"
+#include "StructUtils/StructView.h"
+
+#define LOCTEXT_NAMESPACE "DataLinkJsonToStruct"
+
+void UDataLinkJsonToStruct::OnBuildPins(FDataLinkPinBuilder& Inputs, FDataLinkPinBuilder& Outputs) const
+{
+	Super::OnBuildPins(Inputs, Outputs);
+
+	Inputs.Add(UE::DataLinkJson::InputJsonObject)
+		.SetDisplayName(LOCTEXT("InputJsonObjectDisplay", "Json"))
+		.SetStruct<FJsonObjectWrapper>();
+
+	Inputs.Add(UE::DataLinkJson::InputMappingConfig)
+		.SetDisplayName(LOCTEXT("InputJsonStructMappingConfigDisplay", "Mapping Config"))
+		.SetStruct<FDataLinkJsonStructMappingConfig>();
+
+	// Output does not have a struct defined as it's defined via the Input Mapping Config
+	Outputs.Add(UE::DataLink::OutputDefault)
+		.SetDisplayName(LOCTEXT("OutputDisplay", "Struct"));
+}
+
+EDataLinkExecutionReply UDataLinkJsonToStruct::OnExecute(FDataLinkExecutor& InExecutor) const
+{
+	const FDataLinkNodeInstance& NodeInstance = InExecutor.GetNodeInstance(this);
+
+	const FDataLinkInputDataViewer& InputDataViewer = NodeInstance.GetInputDataViewer();
+
+	const FJsonObjectWrapper& InputData = InputDataViewer.Get<FJsonObjectWrapper>(UE::DataLinkJson::InputJsonObject);
+
+	const FDataLinkJsonStructMappingConfig& MappingConfig = InputDataViewer.Get<FDataLinkJsonStructMappingConfig>(UE::DataLinkJson::InputMappingConfig);
+
+	if (!InputData.JsonObject.IsValid() || !MappingConfig.OutputStruct)
+	{
+		return EDataLinkExecutionReply::Unhandled;
+	}
+
+	bool bResult;
+
+	const FStructView OutputDataView = NodeInstance.GetOutputDataViewer().Find(UE::DataLink::OutputDefault, MappingConfig.OutputStruct);
+
+	if (MappingConfig.CustomMapping)
+	{
+		bResult = MappingConfig.CustomMapping->Apply(InputData.JsonObject.ToSharedRef(), OutputDataView);
+	}
+	else
+	{
+		bResult = FJsonObjectConverter::JsonObjectToUStruct(InputData.JsonObject.ToSharedRef()
+			, OutputDataView.GetScriptStruct()
+			, OutputDataView.GetMemory());
+	}
+
+	if (!bResult)
+	{
+		return EDataLinkExecutionReply::Unhandled;
+	}
+
+	InExecutor.SucceedNode(this, OutputDataView);
+	return EDataLinkExecutionReply::Handled;
+}
+
+#undef LOCTEXT_NAMESPACE
diff --git a/Engine/Plugins/Experimental/DataLink/Source/DataLinkJson/Private/DataLinkJsonUtils.cpp b/Engine/Plugins/Experimental/DataLink/Source/DataLinkJson/Private/DataLinkJsonUtils.cpp
new file mode 100644
--- /dev/null
+++ b/Engine/Plugins/Experimental/DataLink/Source/DataLinkJson/Private/DataLinkJsonUtils.cpp
@@ -0,0 +1,132 @@
+// Copyright Epic Games, Inc. All Rights Reserved.
+
+#include "DataLinkJsonUtils.h"
+#include "DataLinkJsonLog.h"
+#include "Dom/JsonObject.h"
+#include "Dom/JsonValue.h"
+
+namespace UE::DataLinkJson::Private
+{
+	constexpr FStringView LeftArrayDelimiter = TEXT("[");
+	constexpr FStringView RightArrayDelimiter = TEXT("]");
+
+	void FindArraySubscripts(FStringView InStringView, TArray<FStringView>& OutSubscriptViews)
+	{
+		while (!InStringView.IsEmpty())
+		{
+			const int32 StartIndex = UE::String::FindFirst(InStringView, LeftArrayDelimiter, ESearchCase::CaseSensitive);
+			if (StartIndex == INDEX_NONE)
+			{
+				return;
+			}
+
+			// Move string to the right so that it no longer has the left delimiter in view.
+			InStringView.RightChopInline(StartIndex + LeftArrayDelimiter.Len());
+			if (InStringView.IsEmpty())
+			{
+				return;
+			}
+
+			const int32 EndIndex = UE::String::FindFirst(InStringView, RightArrayDelimiter, ESearchCase::CaseSensitive);
+			if (EndIndex == INDEX_NONE)
+			{
+				return;
+			}
+
+			OutSubscriptViews.Add(InStringView.Left(EndIndex));
+			InStringView.RightChopInline(EndIndex + RightArrayDelimiter.Len());
+		}
+	}
+
+	bool IterateArraySubscripts(TSharedPtr<FJsonValue>& InOutValue, TConstArrayView<FStringView> InArraySubscripts, const FString& InFieldName)
+	{
+		constexpr const TCHAR* ErrorPrefix = TEXT("Error while iterating Array Subscripts -");
+
+		for (const FStringView& ArraySubscript : InArraySubscripts)
+		{
+			if (!InOutValue.IsValid())
+			{
+				UE_LOG(LogDataLinkJson, Log, TEXT("%s '%s' Json Value is null."), ErrorPrefix, *InFieldName);
+				return false;
+			}
+
+			TArray<TSharedPtr<FJsonValue>>* ValueArray;
+			if (!InOutValue->TryGetArray(ValueArray))
+			{
+				UE_LOG(LogDataLinkJson, Log, TEXT("%s '%s' is not an array!"), ErrorPrefix, *InFieldName);
+				return false;
+			}
+
+			check(ValueArray);
+			FStringView::ElementType* End = nullptr;
+			const int32 ArrayIndex = FCString::Strtoi(ArraySubscript.GetData(), &End, /*Base*/10);
+
+			if (End != ArraySubscript.end())
+			{
+				UE_LOG(LogDataLinkJson, Log, TEXT("%s Array '%s' with non-numeric subscript '%s'."), ErrorPrefix, *InFieldName, *FString(ArraySubscript));
+				return false;
+			}
+
+			if (!ValueArray->IsValidIndex(ArrayIndex))
+			{
+				UE_LOG(LogDataLinkJson, Log, TEXT("%s Array '%s' does not have a valid index %d."), ErrorPrefix, *InFieldName, ArrayIndex);
+				return false;
+			}
+
+			InOutValue = (*ValueArray)[ArrayIndex];
+		}
+		return true;
+	}
+}
+
+TSharedPtr<FJsonValue> UE::DataLinkJson::FindJsonValue(const TSharedRef<FJsonObject>& InJsonObject, const FString& InFieldName)
+{
+	TArray<FString> FieldPath;
+	InFieldName.ParseIntoArray(FieldPath, TEXT("."));
+
+	// Rather than removing first item, reverse and remove from end
+	Algo::Reverse(FieldPath);
+
+	TSharedPtr<FJsonObject> CurrentObject = InJsonObject;
+	while (!FieldPath.IsEmpty() && CurrentObject.IsValid())
+	{
+		const FString FieldName = FieldPath.Pop();
+
+		FStringView FieldNameView = FieldName;
+
+		TArray<FStringView> ArraySubscripts;
+		Private::FindArraySubscripts(FieldName, ArraySubscripts);
+
+		// If there were subscripts, only consider the left part of the string
+		if (!ArraySubscripts.IsEmpty())
+		{
+			const int32 FirstSubscriptIndex = ArraySubscripts[0].GetData() - FieldNameView.GetData();
+			FieldNameView.LeftChopInline(FirstSubscriptIndex - Private::LeftArrayDelimiter.Len());
+		}
+
+		TSharedPtr<FJsonValue> Value = CurrentObject->TryGetField(FieldNameView);
+		if (!Value.IsValid())
+		{
+			UE_LOG(LogDataLinkJson, Log, TEXT("Failed to find Json Value %s."), *FieldName);
+			return nullptr;
+		}
+
+		// Parse array if there are any subscripts (will skip if empty)
+		if (!Private::IterateArraySubscripts(Value, ArraySubscripts, FieldName))
+		{
+			return nullptr;
+		}
+
+		// If last segment of the path, the value was found 
+		if (FieldPath.IsEmpty())
+		{
+			return Value;
+		}
+
+		// Else, there's more segments to look into.
+		// The value should be an object if there's going to be further diving
+		CurrentObject = Value->AsObject();
+	}
+
+	return nullptr;
+}
diff --git a/Engine/Plugins/Experimental/DataLink/Source/DataLinkJson/Private/DataLinkStringToJson.cpp b/Engine/Plugins/Experimental/DataLink/Source/DataLinkJson/Private/DataLinkStringToJson.cpp
new file mode 100644
--- /dev/null
+++ b/Engine/Plugins/Experimental/DataLink/Source/DataLinkJson/Private/DataLinkStringToJson.cpp
@@ -0,0 +1,57 @@
+// Copyright Epic Games, Inc. All Rights Reserved.
+
+#include "DataLinkStringToJson.h"
+#include "DataLinkCoreTypes.h"
+#include "DataLinkExecutor.h"
+#include "DataLinkJsonNames.h"
+#include "DataLinkNames.h"
+#include "DataLinkNodeInstance.h"
+#include "DataLinkPinBuilder.h"
+#include "JsonObjectWrapper.h"
+#include "StructUtils/StructView.h"
+
+#define LOCTEXT_NAMESPACE "DataLinkStringToJson"
+
+void UDataLinkStringToJson::OnBuildPins(FDataLinkPinBuilder& Inputs, FDataLinkPinBuilder& Outputs) const
+{
+	Super::OnBuildPins(Inputs, Outputs);
+
+	Inputs.Add(UE::DataLinkJson::InputString)
+		.SetDisplayName(LOCTEXT("InputString", "String"))
+		.SetStruct<FDataLinkString>();
+
+	Outputs.Add(UE::DataLink::OutputDefault)
+		.SetDisplayName(LOCTEXT("OutputDisplay", "Json"))
+		.SetStruct<FJsonObjectWrapper>();
+}
+
+EDataLinkExecutionReply UDataLinkStringToJson::OnExecute(FDataLinkExecutor& InExecutor) const
+{
+	const FDataLinkNodeInstance& NodeInstance = InExecutor.GetNodeInstance(this);
+
+	const FDataLinkInputDataViewer& InputDataViewer = NodeInstance.GetInputDataViewer();
+	const FDataLinkOutputDataViewer& OutputDataViewer = NodeInstance.GetOutputDataViewer();
+
+	const FDataLinkString& InputData = InputDataViewer.Get<FDataLinkString>(UE::DataLinkJson::InputString);
+
+	FJsonObjectWrapper& OutputData = OutputDataViewer.Get<FJsonObjectWrapper>(UE::DataLink::OutputDefault);
+
+	// Return early if the output data's json string already matches the input string
+	if (OutputData.JsonString == InputData.Value)
+	{
+		InExecutor.SucceedNode(this, FConstStructView::Make(OutputData));
+		return EDataLinkExecutionReply::Handled;
+	}
+
+	if (!OutputData.JsonObjectFromString(InputData.Value))
+	{
+		return EDataLinkExecutionReply::Unhandled;
+	}
+
+	// Save the Json String for future reference if the parsed object can be re-used for matching strings
+	OutputData.JsonString = InputData.Value;
+	InExecutor.SucceedNode(this, FConstStructView::Make(OutputData));
+	return EDataLinkExecutionReply::Handled;
+}
+
+#undef LOCTEXT_NAMESPACE
diff --git a/Engine/Plugins/Experimental/DataLink/Source/DataLinkJson/Public/DataLinkJsonFromStruct.h b/Engine/Plugins/Experimental/DataLink/Source/DataLinkJson/Public/DataLinkJsonFromStruct.h
new file mode 100644
--- /dev/null
+++ b/Engine/Plugins/Experimental/DataLink/Source/DataLinkJson/Public/DataLinkJsonFromStruct.h
@@ -0,0 +1,19 @@
+// Copyright Epic Games, Inc. All Rights Reserved.
+
+#pragma once
+
+#include "DataLinkNode.h"
+#include "DataLinkJsonFromStruct.generated.h"
+
+/** Convert a struct to Json Object */
+UCLASS(MinimalAPI, Category="JSON", DisplayName="Struct to JSON")
+class UDataLinkJsonFromStruct : public UDataLinkNode
+{
+	GENERATED_BODY()
+
+protected:
+	//~ Begin UDataLinkNode
+	DATALINKJSON_API virtual void OnBuildPins(FDataLinkPinBuilder& Inputs, FDataLinkPinBuilder& Outputs) const override;
+	DATALINKJSON_API virtual EDataLinkExecutionReply OnExecute(FDataLinkExecutor& InExecutor) const override;
+	//~ End UDataLinkNode
+};
diff --git a/Engine/Plugins/Experimental/DataLink/Source/DataLinkJson/Public/DataLinkJsonLog.h b/Engine/Plugins/Experimental/DataLink/Source/DataLinkJson/Public/DataLinkJsonLog.h
new file mode 100644
--- /dev/null
+++ b/Engine/Plugins/Experimental/DataLink/Source/DataLinkJson/Public/DataLinkJsonLog.h
@@ -0,0 +1,7 @@
+// Copyright Epic Games, Inc. All Rights Reserved.
+
+#pragma once
+
+#include "Logging/LogMacros.h"
+
+DECLARE_LOG_CATEGORY_EXTERN(LogDataLinkJson, Log, All);
diff --git a/Engine/Plugins/Experimental/DataLink/Source/DataLinkJson/Public/DataLinkJsonNames.h b/Engine/Plugins/Experimental/DataLink/Source/DataLinkJson/Public/DataLinkJsonNames.h
new file mode 100644
--- /dev/null
+++ b/Engine/Plugins/Experimental/DataLink/Source/DataLinkJson/Public/DataLinkJsonNames.h
@@ -0,0 +1,13 @@
+// Copyright Epic Games, Inc. All Rights Reserved.
+
+#pragma once
+
+#include "UObject/NameTypes.h"
+
+namespace UE::DataLinkJson
+{
+	const FLazyName InputJsonObject(TEXT("InputJson"));
+	const FLazyName InputMappingConfig(TEXT("InputMapping"));
+	const FLazyName InputString(TEXT("InputString"));
+	const FLazyName InputStruct(TEXT("InputStruct"));
+}
diff --git a/Engine/Plugins/Experimental/DataLink/Source/DataLinkJson/Public/DataLinkJsonStructMapping.h b/Engine/Plugins/Experimental/DataLink/Source/DataLinkJson/Public/DataLinkJsonStructMapping.h
new file mode 100644
--- /dev/null
+++ b/Engine/Plugins/Experimental/DataLink/Source/DataLinkJson/Public/DataLinkJsonStructMapping.h
@@ -0,0 +1,29 @@
+// Copyright Epic Games, Inc. All Rights Reserved.
+
+#pragma once
+
+#include "Templates/SharedPointerFwd.h"
+#include "UObject/Object.h"
+#include "DataLinkJsonStructMapping.generated.h"
+
+class FJsonObject;
+struct FStructView;
+
+/** Base class for any mapping logic from a json object to a given struct */
+UCLASS(MinimalAPI, Abstract, EditInlineNew)
+class UDataLinkJsonStructMapping : public UObject
+{
+	GENERATED_BODY()
+
+public:
+	/**
+	 * Runs the mapping logic on the given Json object and the target struct view
+	 * @param InSourceJson the json object to convert
+	 * @param InTargetStructView the destination struct instance
+	 * @return true if the mapping succeeded
+	 */
+	DATALINKJSON_API virtual bool Apply(const TSharedRef<FJsonObject>& InSourceJson, const FStructView& InTargetStructView) const
+	{
+		return false;
+	}
+};
diff --git a/Engine/Plugins/Experimental/DataLink/Source/DataLinkJson/Public/DataLinkJsonStructSimpleMapping.h b/Engine/Plugins/Experimental/DataLink/Source/DataLinkJson/Public/DataLinkJsonStructSimpleMapping.h
new file mode 100644
--- /dev/null
+++ b/Engine/Plugins/Experimental/DataLink/Source/DataLinkJson/Public/DataLinkJsonStructSimpleMapping.h
@@ -0,0 +1,25 @@
+// Copyright Epic Games, Inc. All Rights Reserved.
+
+#pragma once
+
+#include "DataLinkJsonStructMapping.h"
+#include "DataLinkJsonStructSimpleMapping.generated.h"
+
+UCLASS(MinimalAPI)
+class UDataLinkJsonStructSimpleMapping : public UDataLinkJsonStructMapping
+{
+	GENERATED_BODY()
+
+public:
+	//~ Begin UDataLinkJsonStructMapping
+	DATALINKJSON_API virtual bool Apply(const TSharedRef<FJsonObject>& InSourceJson, const FStructView& InTargetStructView) const override;
+	//~ End UDataLinkJsonStructMapping
+
+protected:
+	/**
+	 * Map for how a field in the json pairs to a given property name in a struct.
+	 * Nested json fields can be accessed via a dot delimiter for each path segment.
+	 */
+	UPROPERTY(EditAnywhere, Category="Data Link")
+	TMap<FString, FName> FieldMappings;
+};
diff --git a/Engine/Plugins/Experimental/DataLink/Source/DataLinkJson/Public/DataLinkJsonToStruct.h b/Engine/Plugins/Experimental/DataLink/Source/DataLinkJson/Public/DataLinkJsonToStruct.h
new file mode 100644
--- /dev/null
+++ b/Engine/Plugins/Experimental/DataLink/Source/DataLinkJson/Public/DataLinkJsonToStruct.h
@@ -0,0 +1,39 @@
+// Copyright Epic Games, Inc. All Rights Reserved.
+
+#pragma once
+
+#include "DataLinkNode.h"
+#include "DataLinkJsonToStruct.generated.h"
+
+class UDataLinkJsonStructMapping;
+
+USTRUCT(BlueprintType)
+struct FDataLinkJsonStructMappingConfig
+{
+	GENERATED_BODY()
+
+	/** The desired struct to convert the json object to */
+	UPROPERTY(EditAnywhere, Category="Data Link")
+	TObjectPtr<UScriptStruct> OutputStruct;
+
+	/**
+	 * Optional custom mapping to handle converting the Json object to the Output Struct
+	 * If none is specified, the default mapping method will be used where the Struct property hierarchy should match that of the Json's
+	 * @see FJsonObjectConverter::JsonObjectToUStruct
+	 */
+	UPROPERTY(EditAnywhere, Instanced, Category="Data Link")
+	TObjectPtr<UDataLinkJsonStructMapping> CustomMapping;
+};
+
+/** Convert a Json Object to a particular struct */
+UCLASS(MinimalAPI, Category="JSON", DisplayName="JSON to Struct")
+class UDataLinkJsonToStruct : public UDataLinkNode
+{
+	GENERATED_BODY()
+
+protected:
+	//~ Begin UDataLinkNode
+	DATALINKJSON_API virtual void OnBuildPins(FDataLinkPinBuilder& Inputs, FDataLinkPinBuilder& Outputs) const override;
+	DATALINKJSON_API virtual EDataLinkExecutionReply OnExecute(FDataLinkExecutor& InExecutor) const override;
+	//~ End UDataLinkNode
+};
diff --git a/Engine/Plugins/Experimental/DataLink/Source/DataLinkJson/Public/DataLinkJsonUtils.h b/Engine/Plugins/Experimental/DataLink/Source/DataLinkJson/Public/DataLinkJsonUtils.h
new file mode 100644
--- /dev/null
+++ b/Engine/Plugins/Experimental/DataLink/Source/DataLinkJson/Public/DataLinkJsonUtils.h
@@ -0,0 +1,20 @@
+// Copyright Epic Games, Inc. All Rights Reserved.
+
+#pragma once
+
+#include "Templates/SharedPointerFwd.h"
+
+class FJsonObject;
+class FJsonValue;
+class FString;
+
+namespace UE::DataLinkJson
+{
+	/**
+	 * Finds the json value that corresponds to the given field name in the given json object.
+	 * @param InJsonObject the json object to look into
+	 * @param InFieldName the field name to match. Can be in the format "A.B[2].C" to return nested values and array elements
+	 * @return the json value if found
+	 */
+	DATALINKJSON_API TSharedPtr<FJsonValue> FindJsonValue(const TSharedRef<FJsonObject>& InJsonObject, const FString& InFieldName);
+}
diff --git a/Engine/Plugins/Experimental/DataLink/Source/DataLinkJson/Public/DataLinkStringToJson.h b/Engine/Plugins/Experimental/DataLink/Source/DataLinkJson/Public/DataLinkStringToJson.h
new file mode 100644
--- /dev/null
+++ b/Engine/Plugins/Experimental/DataLink/Source/DataLinkJson/Public/DataLinkStringToJson.h
@@ -0,0 +1,19 @@
+// Copyright Epic Games, Inc. All Rights Reserved.
+
+#pragma once
+
+#include "DataLinkNode.h"
+#include "DataLinkStringToJson.generated.h"
+
+/** Convert a String to a Json Object */
+UCLASS(MinimalAPI, Category="JSON", DisplayName="String to JSON")
+class UDataLinkStringToJson : public UDataLinkNode
+{
+	GENERATED_BODY()
+
+protected:
+	//~ Begin UDataLinkNode
+	DATALINKJSON_API virtual void OnBuildPins(FDataLinkPinBuilder& Inputs, FDataLinkPinBuilder& Outputs) const override;
+	DATALINKJSON_API virtual EDataLinkExecutionReply OnExecute(FDataLinkExecutor& InExecutor) const override;
+	//~ End UDataLinkNode
+};
diff --git a/Engine/Plugins/Experimental/DataLink/Source/DataLinkJsonEditor/DataLinkJsonEditor.Build.cs b/Engine/Plugins/Experimental/DataLink/Source/DataLinkJsonEditor/DataLinkJsonEditor.Build.cs
new file mode 100644
--- /dev/null
+++ b/Engine/Plugins/Experimental/DataLink/Source/DataLinkJsonEditor/DataLinkJsonEditor.Build.cs
@@ -0,0 +1,30 @@
+// Copyright Epic Games, Inc. All Rights Reserved.
+
+using UnrealBuildTool;
+
+public class DataLinkJsonEditor : ModuleRules
+{
+    public DataLinkJsonEditor(ReadOnlyTargetRules Target) : base(Target)
+    {
+        PCHUsage = ModuleRules.PCHUsageMode.UseExplicitOrSharedPCHs;
+
+        PrivateDependencyModuleNames.AddRange(
+            new string[]
+            {
+                "AssetTools",
+                "BlueprintGraph",
+                "ContentBrowser",
+                "Core",
+                "CoreUObject",
+                "DataLinkEditor",
+                "Engine",
+                "Json",
+                "JsonUtilities",
+                "Slate",
+                "SlateCore",
+                "ToolMenus",
+                "UnrealEd",
+            }
+        );
+    }
+}
\ No newline at end of file
diff --git a/Engine/Plugins/Experimental/DataLink/Source/DataLinkJsonEditor/Private/DataLinkJsonEditorLog.h b/Engine/Plugins/Experimental/DataLink/Source/DataLinkJsonEditor/Private/DataLinkJsonEditorLog.h
new file mode 100644
--- /dev/null
+++ b/Engine/Plugins/Experimental/DataLink/Source/DataLinkJsonEditor/Private/DataLinkJsonEditorLog.h
@@ -0,0 +1,7 @@
+// Copyright Epic Games, Inc. All Rights Reserved.
+
+#pragma once
+
+#include "Logging/LogMacros.h"
+
+DECLARE_LOG_CATEGORY_EXTERN(LogDataLinkJsonEditor, Log, All);
diff --git a/Engine/Plugins/Experimental/DataLink/Source/DataLinkJsonEditor/Private/DataLinkJsonEditorMenu.cpp b/Engine/Plugins/Experimental/DataLink/Source/DataLinkJsonEditor/Private/DataLinkJsonEditorMenu.cpp
new file mode 100644
--- /dev/null
+++ b/Engine/Plugins/Experimental/DataLink/Source/DataLinkJsonEditor/Private/DataLinkJsonEditorMenu.cpp
@@ -0,0 +1,104 @@
+// Copyright Epic Games, Inc. All Rights Reserved.
+
+#include "DataLinkJsonEditorMenu.h"
+#include "DataLinkEditorNames.h"
+#include "DataLinkJsonEditorStructGenerator.h"
+#include "Dialogs/DlgPickPath.h"
+#include "Framework/Application/SlateApplication.h"
+#include "IDataLinkEditorMenuContext.h"
+#include "JsonObjectWrapper.h"
+#include "StructUtils/StructView.h"
+#include "Styling/AppStyle.h"
+#include "ToolMenu.h"
+#include "ToolMenuSection.h"
+#include "ToolMenus.h"
+#include "Widgets/SDataLinkJsonEditorStructGeneratorConfig.h"
+
+#define LOCTEXT_NAMESPACE "DataLinkJsonEditorMenu"
+
+void UE::DataLinkJsonEditor::RegisterMenus()
+{
+	UToolMenu* ToolMenu = UToolMenus::Get()->ExtendMenu(UE::DataLinkEditor::PreviewToolbarName);
+	if (!ToolMenu)
+	{
+		return;
+	}
+
+	ToolMenu->AddDynamicSection(TEXT("JsonEditor")
+		, FNewToolMenuDelegate::CreateStatic(&UE::DataLinkJsonEditor::PopulateToolbar)
+		, FToolMenuInsert(UE::DataLinkEditor::PreviewSectionName, EToolMenuInsertType::After));
+}
+
+void UE::DataLinkJsonEditor::PopulateToolbar(UToolMenu* InToolMenu)
+{
+	FToolUIAction MakeStructsAction;
+	MakeStructsAction.ExecuteAction.BindStatic(&UE::DataLinkJsonEditor::MakeStructsFromJson);
+	MakeStructsAction.CanExecuteAction.BindStatic(&UE::DataLinkJsonEditor::CanMakeStructsFromJson);
+
+	FToolMenuSection& Section = InToolMenu->AddSection(TEXT("Json"));
+	Section.AddEntry(FToolMenuEntry::InitToolBarButton(TEXT("MakeStructs")
+		, MakeStructsAction
+		, LOCTEXT("MakeStructsDisplayName", "Make Structs from Json")
+		, LOCTEXT("Tooltip", "Makes struct assets matching the Json hierarchy of the Output Data (needs to be Json)")
+		, FSlateIcon(FAppStyle::GetAppStyleSetName(), TEXT("Icons.Toolbar.Export"))));
+}
+
+TConstStructView<FJsonObjectWrapper> UE::DataLinkJsonEditor::GetPreviewOutputDataView(const IDataLinkEditorMenuContext& InMenuContext)
+{
+	FConstStructView PreviewOutputData = InMenuContext.FindPreviewOutputData();
+	if (!PreviewOutputData.IsValid() || PreviewOutputData.GetScriptStruct() != FJsonObjectWrapper::StaticStruct())
+	{
+		return TConstStructView<FJsonObjectWrapper>();
+	}
+
+	return TConstStructView<FJsonObjectWrapper>(PreviewOutputData.GetMemory());
+}
+
+bool UE::DataLinkJsonEditor::CanMakeStructsFromJson(const FToolMenuContext& InMenuContext)
+{
+	IDataLinkEditorMenuContext* const MenuContext = InMenuContext.FindContext<IDataLinkEditorMenuContext>();
+	if (!MenuContext)
+	{
+		return false;
+	}
+
+	TConstStructView<FJsonObjectWrapper> JsonWrapper = GetPreviewOutputDataView(*MenuContext);
+	return JsonWrapper.IsValid() && JsonWrapper.Get().JsonObject.IsValid();
+}
+
+void UE::DataLinkJsonEditor::MakeStructsFromJson(const FToolMenuContext& InMenuContext)
+{
+	IDataLinkEditorMenuContext* const MenuContext = InMenuContext.FindContext<IDataLinkEditorMenuContext>();
+	if (!MenuContext)
+	{
+		return;
+	}
+
+	TConstStructView<FJsonObjectWrapper> JsonWrapper = GetPreviewOutputDataView(*MenuContext);
+	if (!JsonWrapper.IsValid())
+	{
+		return;
+	}
+
+	TSharedPtr<FJsonObject> JsonObject = JsonWrapper.Get().JsonObject;
+	if (!JsonObject.IsValid())
+	{
+		return;
+	}
+
+	FSlateApplication::Get().AddWindow(SNew(SDataLinkJsonEditorStructGeneratorConfig)
+		.Title(LOCTEXT("ChooseTargetContentPath", "Choose Location for the Structs"))
+		.DefaultPath(MenuContext->GetAssetPath())
+		.OnCommit_Lambda(
+			[JsonObject = MoveTemp(JsonObject)](const SDataLinkJsonEditorStructGeneratorConfig& InConfig)
+			{
+				FStructGenerator::FParams Params;
+				Params.JsonObject = JsonObject;
+				Params.BasePath = InConfig.GetPath();
+				Params.StructPrefix = InConfig.GetPrefix();
+				Params.RootStructName = InConfig.GetRootStructName();
+				FStructGenerator::GenerateFromJson(Params);
+			}));
+}
+
+#undef LOCTEXT_NAMESPACE
diff --git a/Engine/Plugins/Experimental/DataLink/Source/DataLinkJsonEditor/Private/DataLinkJsonEditorMenu.h b/Engine/Plugins/Experimental/DataLink/Source/DataLinkJsonEditor/Private/DataLinkJsonEditorMenu.h
new file mode 100644
--- /dev/null
+++ b/Engine/Plugins/Experimental/DataLink/Source/DataLinkJsonEditor/Private/DataLinkJsonEditorMenu.h
@@ -0,0 +1,23 @@
+// Copyright Epic Games, Inc. All Rights Reserved.
+
+#pragma once
+
+class FJsonObject;
+class IDataLinkEditorMenuContext;
+class UToolMenu;
+struct FJsonObjectWrapper;
+struct FToolMenuContext;
+template<typename BaseStructT> struct TConstStructView;
+
+namespace UE::DataLinkJsonEditor
+{
+	void RegisterMenus();
+
+	void PopulateToolbar(UToolMenu* InToolMenu);
+
+	TConstStructView<FJsonObjectWrapper> GetPreviewOutputDataView(const IDataLinkEditorMenuContext& InMenuContext);
+
+	bool CanMakeStructsFromJson(const FToolMenuContext& InMenuContext);
+
+	void MakeStructsFromJson(const FToolMenuContext& InMenuContext);
+}
diff --git a/Engine/Plugins/Experimental/DataLink/Source/DataLinkJsonEditor/Private/DataLinkJsonEditorModule.cpp b/Engine/Plugins/Experimental/DataLink/Source/DataLinkJsonEditor/Private/DataLinkJsonEditorModule.cpp
new file mode 100644
--- /dev/null
+++ b/Engine/Plugins/Experimental/DataLink/Source/DataLinkJsonEditor/Private/DataLinkJsonEditorModule.cpp
@@ -0,0 +1,26 @@
+// Copyright Epic Games, Inc. All Rights Reserved.
+
+#include "DataLinkJsonEditorModule.h"
+#include "DataLinkEditorNames.h"
+#include "DataLinkJsonEditorLog.h"
+#include "DataLinkJsonEditorMenu.h"
+#include "Modules/ModuleManager.h"
+#include "ToolMenus.h"
+
+DEFINE_LOG_CATEGORY(LogDataLinkJsonEditor);
+
+IMPLEMENT_MODULE(FDataLinkJsonEditorModule, DataLinkJsonEditor)
+
+void FDataLinkJsonEditorModule::StartupModule()
+{
+	FSimpleMulticastDelegate::FDelegate OnToolMenusStartup;
+	OnToolMenusStartup.BindStatic(&UE::DataLinkJsonEditor::RegisterMenus);
+
+	OnToolMenusStartupHandle = UToolMenus::RegisterStartupCallback(OnToolMenusStartup);
+}
+
+void FDataLinkJsonEditorModule::ShutdownModule()
+{
+	UToolMenus::UnRegisterStartupCallback(OnToolMenusStartupHandle);
+	OnToolMenusStartupHandle.Reset();
+}
diff --git a/Engine/Plugins/Experimental/DataLink/Source/DataLinkJsonEditor/Private/DataLinkJsonEditorModule.h b/Engine/Plugins/Experimental/DataLink/Source/DataLinkJsonEditor/Private/DataLinkJsonEditorModule.h
new file mode 100644
--- /dev/null
+++ b/Engine/Plugins/Experimental/DataLink/Source/DataLinkJsonEditor/Private/DataLinkJsonEditorModule.h
@@ -0,0 +1,18 @@
+// Copyright Epic Games, Inc. All Rights Reserved.
+
+#pragma once
+
+#include "Delegates/IDelegateInstance.h"
+#include "Modules/ModuleInterface.h"
+
+class FDataLinkJsonEditorModule : public IModuleInterface
+{
+public:
+    //~ Begin IModuleInterface
+    virtual void StartupModule() override;
+    virtual void ShutdownModule() override;
+    //~ End IModuleInterface
+
+private:
+    FDelegateHandle OnToolMenusStartupHandle;
+};
diff --git a/Engine/Plugins/Experimental/DataLink/Source/DataLinkJsonEditor/Private/DataLinkJsonEditorStructGenerator.cpp b/Engine/Plugins/Experimental/DataLink/Source/DataLinkJsonEditor/Private/DataLinkJsonEditorStructGenerator.cpp
new file mode 100644
--- /dev/null
+++ b/Engine/Plugins/Experimental/DataLink/Source/DataLinkJsonEditor/Private/DataLinkJsonEditorStructGenerator.cpp
@@ -0,0 +1,224 @@
+// Copyright Epic Games, Inc. All Rights Reserved.
+
+#include "DataLinkJsonEditorStructGenerator.h"
+#include "AssetToolsModule.h"
+#include "DataLinkJsonEditorLog.h"
+#include "Dom/JsonObject.h"
+#include "Dom/JsonValue.h"
+#include "EdGraphSchema_K2.h"
+#include "Factories/StructureFactory.h"
+#include "IAssetTools.h"
+#include "IContentBrowserSingleton.h"
+#include "Kismet2/StructureEditorUtils.h"
+#include "PackageTools.h"
+#include "ScopedTransaction.h"
+#include "UserDefinedStructure/UserDefinedStructEditorData.h"
+
+#define LOCTEXT_NAMESPACE "DataLinkJsonEditorUtils"
+
+UE::DataLinkJsonEditor::FStructKey::FStructKey(FStructGenerator& InStructGenerator, const TMap<FString, TSharedPtr<FJsonValue>>& InJsonEntries)
+{
+	Hash = 0;
+	JsonTypeMap.Reserve(InJsonEntries.Num());
+	for (const TPair<FString, TSharedPtr<FJsonValue>>& JsonEntry : InJsonEntries)
+	{
+		FEdGraphPinType Type;
+		if (InStructGenerator.FromJsonValue(JsonEntry.Value, Type, JsonEntry.Key))
+		{
+			Hash = HashCombine(Hash, GetTypeHash(JsonEntry.Key));
+			Hash = HashCombine(Hash, GetTypeHash(Type.PinCategory));
+			Hash = HashCombine(Hash, GetTypeHash(Type.PinSubCategoryObject));
+			Hash = HashCombine(Hash, GetTypeHash(Type.ContainerType));
+
+			JsonTypeMap.Add(JsonEntry.Key, Type);
+		}
+	}
+}
+
+void UE::DataLinkJsonEditor::FStructGenerator::GenerateFromJson(const FParams& InParams)
+{
+	FScopedSlowTask SlowTask(0, LOCTEXT("SlowTaskText", "Generating structs from Json..."));
+	SlowTask.MakeDialog();
+
+	// Note: this transaction is only here because FStructureEditorUtils CreateAsset for User Defined Structs transact
+	// because of FStructureEditorUtils::AddVariable. So this is here to avoid multiple transactions and control it from here.
+	// This transaction will be cancelled at the end as asset creation should've not transacted in the first place.
+	FScopedTransaction Transaction(LOCTEXT("Transaction", "Generate Structs from Json"));
+
+	FStructGenerator Generator;
+	Generator.BasePath = InParams.BasePath;
+	Generator.StructPrefix = InParams.StructPrefix;
+	Generator.GetOrCreateStruct(InParams.JsonObject, InParams.RootStructName);
+
+	Transaction.Cancel();
+
+	IContentBrowserSingleton::Get().SyncBrowserToAssets(Generator.GetGeneratedStructs());
+}
+
+bool UE::DataLinkJsonEditor::FStructGenerator::FromJsonValue(const TSharedPtr<FJsonValue>& InJsonValue, FEdGraphPinType& OutPinType, const FString& InNameToUse)
+{
+	if (!InJsonValue.IsValid())
+	{
+		UE_LOG(LogDataLinkJsonEditor, Error, TEXT("Invalid Json Value '%s'"), *InNameToUse);
+		return false;
+	}
+
+	switch (InJsonValue->Type)
+	{
+	case EJson::String:
+		OutPinType.PinCategory = UEdGraphSchema_K2::PC_String;
+		return true;
+
+	case EJson::Number:
+		// default to real as seen in `FJsonValue::AsNumber`
+		OutPinType.PinCategory = UEdGraphSchema_K2::PC_Real;
+		OutPinType.PinSubCategory = UEdGraphSchema_K2::PC_Double;
+		return true;
+
+	case EJson::Boolean:
+		OutPinType.PinCategory = UEdGraphSchema_K2::PC_Boolean;
+		return true;
+
+	case EJson::Array:
+		{
+			const TArray<TSharedPtr<FJsonValue>>& JsonArray = InJsonValue->AsArray();
+			if (JsonArray.IsEmpty() || !JsonArray[0].IsValid())
+			{
+				UE_LOG(LogDataLinkJsonEditor, Error, TEXT("Empty array '%s' is not supported!"), *InNameToUse);
+				return false;
+			}
+
+			// Nested arrays are not supported.
+			// A middle struct can be created, but the json utils to import/export structs would not produce same hierarchy as the imported json
+			if (JsonArray[0]->Type == EJson::Array)
+			{
+				UE_LOG(LogDataLinkJsonEditor, Error, TEXT("Nested array '%s' is not supported!"), *InNameToUse);
+				return false;
+			}
+
+			// Use the first element as sample to get the inner array type
+			if (FromJsonValue(JsonArray[0], OutPinType, InNameToUse))
+			{
+				OutPinType.ContainerType = EPinContainerType::Array;
+				return true;
+			}
+		}
+		UE_LOG(LogDataLinkJsonEditor, Error, TEXT("Could not create a valid pin array type for '%s'"), *InNameToUse);
+		return false;
+
+	case EJson::Object:
+		if (UUserDefinedStruct* Struct = GetOrCreateStruct(InJsonValue->AsObject(), InNameToUse))
+		{
+			OutPinType.PinCategory = UEdGraphSchema_K2::PC_Struct;
+			OutPinType.PinSubCategoryObject = Struct;
+			return true;
+		}
+		UE_LOG(LogDataLinkJsonEditor, Error, TEXT("Could not create a valid struct type for '%s'"), *InNameToUse);
+		return false;
+	}
+
+	UE_LOG(LogDataLinkJsonEditor, Error, TEXT("Unsupported Json type encountered at '%s'"), *InNameToUse);
+	return false;
+}
+
+void UE::DataLinkJsonEditor::FStructGenerator::AddVariable(UUserDefinedStruct* InStruct, const FString& InName, const TSharedPtr<FJsonValue>& InValue)
+{
+	FEdGraphPinType PinType;
+	if (!FromJsonValue(InValue, PinType, InName))
+	{
+		return;
+	}
+
+	const FGuid Guid = FGuid::NewGuid();
+
+	FStructVariableDescription VarDesc;
+	VarDesc.VarName = *FString::Printf(TEXT("%s_%s"), *InName, *Guid.ToString(EGuidFormats::Digits));
+	VarDesc.VarGuid = Guid;
+	VarDesc.FriendlyName = InName;
+	VarDesc.SetPinType(PinType);
+
+	FStructureEditorUtils::GetVarDesc(InStruct).Add(MoveTemp(VarDesc));
+}
+
+UUserDefinedStruct* UE::DataLinkJsonEditor::FStructGenerator::CreateEmptyStruct(const FString& InNameToUse)
+{
+	if (InNameToUse.IsEmpty())
+	{
+		UE_LOG(LogDataLinkJsonEditor, Error, TEXT("Could not create a valid struct. Empty name!"));
+		return nullptr;
+	}
+
+	if (!StructureFactory)
+	{
+		StructureFactory = NewObject<UStructureFactory>();
+	}
+
+	if (!FName::IsValidXName(InNameToUse, INVALID_OBJECTNAME_CHARACTERS))
+	{
+		const FString CleanName = MakeObjectNameFromDisplayLabel(InNameToUse, NAME_None).GetPlainNameString();
+		UE_LOG(LogDataLinkJsonEditor, Log, TEXT("Input Struct name '%s' is not valid. Cleaned name to '%s'"), *InNameToUse, *CleanName);
+		return CreateEmptyStruct(CleanName);
+	}
+
+	UUserDefinedStruct* const Struct = Cast<UUserDefinedStruct>(IAssetTools::Get().CreateAsset(InNameToUse
+		, BasePath
+		, UUserDefinedStruct::StaticClass()
+		, StructureFactory));
+
+	if (Struct)
+	{
+		FStructureEditorUtils::GetVarDesc(Struct).Empty();
+	}
+
+	return Struct;
+}
+
+UUserDefinedStruct* UE::DataLinkJsonEditor::FStructGenerator::GetOrCreateStruct(const TSharedPtr<FJsonObject>& InJsonObject, const FString& InNameToUse)
+{
+	if (!InJsonObject.IsValid() || InJsonObject->Values.IsEmpty())
+	{
+		UE_LOG(LogDataLinkJsonEditor, Error, TEXT("Could not create a valid struct. Json Object: '%s'"), *InNameToUse);
+		return nullptr;
+	}
+
+	const FStructKey StructKey(*this, InJsonObject->Values);
+	if (const TObjectPtr<UUserDefinedStruct>* FoundStruct = GeneratedStructs.Find(StructKey))
+	{
+		// This could return null for structs that failed to create
+		return *FoundStruct;
+	}
+
+	UUserDefinedStruct* Struct = CreateEmptyStruct(StructPrefix + InNameToUse);
+
+	// Add even if empty. There was an attempt to create the struct, but it failed. Avoid repeating the same steps for
+	// structs with the same schema
+	GeneratedStructs.Add(StructKey, Struct);
+
+	if (!Struct)
+	{
+		return nullptr;
+	}
+
+	for (const TPair<FString, TSharedPtr<FJsonValue>>& JsonEntry : InJsonObject->Values)
+	{
+		AddVariable(Struct, JsonEntry.Key, JsonEntry.Value);
+	}
+
+	FStructureEditorUtils::CompileStructure(Struct);
+	return Struct;
+}
+
+TArray<UObject*> UE::DataLinkJsonEditor::FStructGenerator::GetGeneratedStructs() const
+{
+	TArray<UObject*> Structs;
+	Structs.Reserve(GeneratedStructs.Num());
+
+	for (const TPair<FStructKey, TObjectPtr<UUserDefinedStruct>>& Pair : GeneratedStructs)
+	{
+		Structs.Add(Pair.Value);
+	}
+
+	return Structs;
+}
+
+#undef LOCTEXT_NAMESPACE
diff --git a/Engine/Plugins/Experimental/DataLink/Source/DataLinkJsonEditor/Private/DataLinkJsonEditorStructGenerator.h b/Engine/Plugins/Experimental/DataLink/Source/DataLinkJsonEditor/Private/DataLinkJsonEditorStructGenerator.h
new file mode 100644
--- /dev/null
+++ b/Engine/Plugins/Experimental/DataLink/Source/DataLinkJsonEditor/Private/DataLinkJsonEditorStructGenerator.h
@@ -0,0 +1,74 @@
+// Copyright Epic Games, Inc. All Rights Reserved.
+
+#pragma once
+
+#include "Containers/UnrealString.h"
+#include "EdGraph/EdGraphPin.h"
+#include "Templates/SharedPointerFwd.h"
+
+class FJsonObject;
+class FJsonValue;
+class UUserDefinedStruct;
+class UStructureFactory;
+
+namespace UE::DataLinkJsonEditor
+{
+
+class FStructGenerator;
+
+/** Identifies a map of Name to Json Value to identify a struct that can be re-used */
+struct FStructKey
+{
+	explicit FStructKey(FStructGenerator& InStructGenerator, const TMap<FString, TSharedPtr<FJsonValue>>& InJsonEntries);
+
+	bool operator==(const FStructKey& InOther) const
+	{
+		return Hash == InOther.Hash && JsonTypeMap.OrderIndependentCompareEqual(InOther.JsonTypeMap);
+	}
+
+	friend uint32 GetTypeHash(const FStructKey& InKey)
+	{
+		return InKey.Hash;
+	}
+
+private:
+	uint32 Hash = 0;
+	TMap<FString, FEdGraphPinType> JsonTypeMap;
+};
+
+/** Generates as many structs as needed to create a hierarchy that matches the Json Object*/
+class FStructGenerator
+{
+	friend struct FStructKey;
+
+public:
+	struct FParams
+	{
+		TSharedPtr<FJsonObject> JsonObject;
+		FString BasePath;
+		FString StructPrefix;
+		FString RootStructName;
+	};
+	static void GenerateFromJson(const FParams& InParams);
+
+private:
+	bool FromJsonValue(const TSharedPtr<FJsonValue>& InJsonValue, FEdGraphPinType& OutPinType, const FString& InNameToUse);
+
+	void AddVariable(UUserDefinedStruct* InStruct, const FString& InName, const TSharedPtr<FJsonValue>& InValue);
+
+	UUserDefinedStruct* CreateEmptyStruct(const FString& InNameToUse);
+
+	UUserDefinedStruct* GetOrCreateStruct(const TSharedPtr<FJsonObject>& InJsonObject, const FString& InNameToUse);
+
+	TArray<UObject*> GetGeneratedStructs() const;
+
+	TMap<FStructKey, TObjectPtr<UUserDefinedStruct>> GeneratedStructs;
+
+	TObjectPtr<UStructureFactory> StructureFactory = nullptr;
+
+	FString BasePath;
+
+	FString StructPrefix;
+};
+
+}
diff --git a/Engine/Plugins/Experimental/DataLink/Source/DataLinkJsonEditor/Private/Widgets/SDataLinkJsonEditorStructGeneratorConfig.cpp b/Engine/Plugins/Experimental/DataLink/Source/DataLinkJsonEditor/Private/Widgets/SDataLinkJsonEditorStructGeneratorConfig.cpp
new file mode 100644
--- /dev/null
+++ b/Engine/Plugins/Experimental/DataLink/Source/DataLinkJsonEditor/Private/Widgets/SDataLinkJsonEditorStructGeneratorConfig.cpp
@@ -0,0 +1,186 @@
+// Copyright Epic Games, Inc. All Rights Reserved.
+
+#include "SDataLinkJsonEditorStructGeneratorConfig.h"
+#include "ContentBrowserModule.h"
+#include "DetailLayoutBuilder.h"
+#include "IContentBrowserSingleton.h"
+#include "Misc/MessageDialog.h"
+#include "Widgets/Input/SButton.h"
+#include "Widgets/Input/SEditableTextBox.h"
+#include "Widgets/Layout/SBox.h"
+#include "Widgets/Layout/SUniformGridPanel.h"
+#include "Widgets/SBoxPanel.h"
+#include "Widgets/Text/STextBlock.h"
+
+#define LOCTEXT_NAMESPACE "SDataLinkJsonEditorStructGeneratorConfig"
+
+void SDataLinkJsonEditorStructGeneratorConfig::Construct(const FArguments& InArgs)
+{
+	Path = InArgs._DefaultPath;
+	RootStructName = TEXT("RootStruct");
+	OnCommitDelegate = InArgs._OnCommit;
+
+	FPathPickerConfig PathPickerConfig;
+	PathPickerConfig.DefaultPath = Path;
+	PathPickerConfig.OnPathSelected = FOnPathSelected::CreateSP(this, &SDataLinkJsonEditorStructGeneratorConfig::OnPathChange);
+	PathPickerConfig.bAddDefaultPath = true;
+	PathPickerConfig.bAllowReadOnlyFolders = false;
+
+	FContentBrowserModule& ContentBrowserModule = FModuleManager::LoadModuleChecked<FContentBrowserModule>("ContentBrowser");
+
+	SWindow::Construct(SWindow::FArguments()
+		.Title(InArgs._Title)
+		.SupportsMinimize(false)
+		.SupportsMaximize(false)
+		.ClientSize(FVector2D(450,450))
+		[
+			SNew(SVerticalBox)
+			+ SVerticalBox::Slot()
+			.Padding(2,2,2,4)
+			.AutoHeight()
+			[
+				SNew(SBorder)
+				.BorderImage(FAppStyle::GetBrush("ToolPanel.GroupBorder"))
+				.Padding(1.f)
+				[
+					SNew(SHorizontalBox)
+					+ SHorizontalBox::Slot()
+					.VAlign(VAlign_Center)
+					[
+						SNew(STextBlock)
+						.Text(LOCTEXT("PrefixLabel", "Prefix"))
+						.Font(IDetailLayoutBuilder::GetDetailFont())
+					]
+					+ SHorizontalBox::Slot()
+					.VAlign(VAlign_Center)
+					[
+						SNew(SEditableTextBox)
+						.HintText(LOCTEXT("PrefixHint", "Enter the string to use as prefix for each struct name."))
+						.Font(IDetailLayoutBuilder::GetDetailFont())
+						.OnTextChanged(this, &SDataLinkJsonEditorStructGeneratorConfig::OnPrefixChange)
+					]
+				]
+			]
+			+ SVerticalBox::Slot()
+			.Padding(2,2,2,4)
+			.AutoHeight()
+			[
+				SNew(SBorder)
+				.BorderImage(FAppStyle::GetBrush("ToolPanel.GroupBorder"))
+				.Padding(1.f)
+				[
+					SNew(SHorizontalBox)
+					+ SHorizontalBox::Slot()
+					.VAlign(VAlign_Center)
+					[
+						SNew(STextBlock)
+						.Text(LOCTEXT("RootStructLabel", "Root Struct Name"))
+						.Font(IDetailLayoutBuilder::GetDetailFont())
+					]
+					+ SHorizontalBox::Slot()
+					.VAlign(VAlign_Center)
+					[
+						SNew(SEditableTextBox)
+						.Text(FText::FromString(RootStructName))
+						.Font(IDetailLayoutBuilder::GetDetailFont())
+						.OnTextChanged(this, &SDataLinkJsonEditorStructGeneratorConfig::OnRootStructNameChange)
+					]
+				]
+			]
+			+ SVerticalBox::Slot()
+			.Padding(2,2,2,4)
+			[
+				SNew(SBorder)
+				.BorderImage(FAppStyle::GetBrush("ToolPanel.GroupBorder"))
+				.Padding(3.f)
+				[
+					ContentBrowserModule.Get().CreatePathPicker(PathPickerConfig)
+				]
+			]
+			+SVerticalBox::Slot()
+			.AutoHeight()
+			.HAlign(HAlign_Right)
+			.VAlign(VAlign_Bottom)
+			[
+				SNew(SUniformGridPanel)
+				.SlotPadding(FAppStyle::GetMargin("StandardDialog.SlotPadding"))
+				.MinDesiredSlotWidth(FAppStyle::GetFloat("StandardDialog.MinDesiredSlotWidth"))
+				.MinDesiredSlotHeight(FAppStyle::GetFloat("StandardDialog.MinDesiredSlotHeight"))
+				+SUniformGridPanel::Slot(0,0)
+				[
+					SNew(SButton)
+					.Text(LOCTEXT("OK", "OK"))
+					.HAlign(HAlign_Center)
+					.ContentPadding(FAppStyle::GetMargin("StandardDialog.ContentPadding"))
+					.OnClicked(this, &SDataLinkJsonEditorStructGeneratorConfig::OnCommit)
+				]
+				+SUniformGridPanel::Slot(1,0)
+				[
+					SNew(SButton)
+					.Text(LOCTEXT("Cancel", "Cancel"))
+					.HAlign(HAlign_Center)
+					.ContentPadding(FAppStyle::GetMargin("StandardDialog.ContentPadding"))
+					.OnClicked(this, &SDataLinkJsonEditorStructGeneratorConfig::OnCancel)
+				]
+			]
+		]);
+}
+
+const FString& SDataLinkJsonEditorStructGeneratorConfig::GetPath() const
+{
+	return Path;
+}
+
+const FString& SDataLinkJsonEditorStructGeneratorConfig::GetPrefix() const
+{
+	return Prefix;
+}
+
+const FString& SDataLinkJsonEditorStructGeneratorConfig::GetRootStructName() const
+{
+	return RootStructName;
+}
+
+FReply SDataLinkJsonEditorStructGeneratorConfig::OnCommit()
+{
+	if (ValidateConfiguration())
+	{
+		OnCommitDelegate.ExecuteIfBound(*this);
+		RequestDestroyWindow();
+	}
+	return FReply::Handled();
+}
+
+FReply SDataLinkJsonEditorStructGeneratorConfig::OnCancel()
+{
+	RequestDestroyWindow();
+	return FReply::Handled();
+}
+
+void SDataLinkJsonEditorStructGeneratorConfig::OnPathChange(const FString& InNewPath)
+{
+	Path = InNewPath;
+}
+
+void SDataLinkJsonEditorStructGeneratorConfig::OnPrefixChange(const FText& InNewPrefix)
+{
+	Prefix = InNewPrefix.ToString();
+}
+
+void SDataLinkJsonEditorStructGeneratorConfig::OnRootStructNameChange(const FText& InNewRootName)
+{
+	RootStructName = InNewRootName.ToString();
+}
+
+bool SDataLinkJsonEditorStructGeneratorConfig::ValidateConfiguration()
+{
+	if (Path.IsEmpty())
+	{
+		FMessageDialog::Open(EAppMsgType::Ok, LOCTEXT("EmptyPath", "You must select a path."));
+		return false;
+	}
+
+	return true;
+}
+
+#undef LOCTEXT_NAMESPACE
diff --git a/Engine/Plugins/Experimental/DataLink/Source/DataLinkJsonEditor/Private/Widgets/SDataLinkJsonEditorStructGeneratorConfig.h b/Engine/Plugins/Experimental/DataLink/Source/DataLinkJsonEditor/Private/Widgets/SDataLinkJsonEditorStructGeneratorConfig.h
new file mode 100644
--- /dev/null
+++ b/Engine/Plugins/Experimental/DataLink/Source/DataLinkJsonEditor/Private/Widgets/SDataLinkJsonEditorStructGeneratorConfig.h
@@ -0,0 +1,43 @@
+// Copyright Epic Games, Inc. All Rights Reserved.
+
+#pragma once
+
+#include "Widgets/SWindow.h"
+
+class SDataLinkJsonEditorStructGeneratorConfig : public SWindow
+{
+public:
+	DECLARE_DELEGATE_OneParam(FOnCommit, const SDataLinkJsonEditorStructGeneratorConfig&);
+
+	SLATE_BEGIN_ARGS(SDataLinkJsonEditorStructGeneratorConfig) {}
+	SLATE_ARGUMENT(FText, Title)
+	SLATE_ARGUMENT(FString, DefaultPath)
+	SLATE_EVENT(FOnCommit, OnCommit)
+	SLATE_END_ARGS()
+
+	void Construct(const FArguments& InArgs);
+
+	const FString& GetPath() const;
+	const FString& GetPrefix() const;
+	const FString& GetRootStructName() const;
+
+private:
+	FReply OnCommit();
+	FReply OnCancel();
+
+	void OnPathChange(const FString& InNewPath);
+
+	void OnPrefixChange(const FText& InNewPrefix);
+
+	void OnRootStructNameChange(const FText& InNewRootName);
+
+	bool ValidateConfiguration();
+
+	FString Path;
+
+	FString Prefix;
+
+	FString RootStructName;
+
+	FOnCommit OnCommitDelegate;
+};
