Motion Design Data Link: fixed crash when running data link http request node with invalid settings

diff --git a/Engine/Plugins/Experimental/DataLink/Source/DataLink/Private/DataLinkExecutor.cpp b/Engine/Plugins/Experimental/DataLink/Source/DataLink/Private/DataLinkExecutor.cpp
--- a/Engine/Plugins/Experimental/DataLink/Source/DataLink/Private/DataLinkExecutor.cpp
+++ b/Engine/Plugins/Experimental/DataLink/Source/DataLink/Private/DataLinkExecutor.cpp
@@ -106,6 +106,11 @@ FDataLinkNodeInstance& FDataLinkExecutor::GetNodeInstanceMutable(const UDataLink
 	return NodeInstanceMap.FindChecked(InNode);
 }
 
+FDataLinkNodeInstance* FDataLinkExecutor::FindNodeInstanceMutable(const UDataLinkNode* InNode)
+{
+	return NodeInstanceMap.Find(InNode);
+}
+
 void FDataLinkExecutor::SucceedNode(const UDataLinkNode* InNode, FConstStructView InOutputDataView)
 {
 	if (!ProcessFinishedNode(InNode))
@@ -375,13 +380,20 @@ FDataLinkNodeInstance& FDataLinkExecutor::FindOrAddNodeInstance(const UDataLinkN
 
 bool FDataLinkExecutor::ProcessFinishedNode(const UDataLinkNode* InNode)
 {
-	FDataLinkNodeInstance& NodeInstance = GetNodeInstanceMutable(InNode);
+	FDataLinkNodeInstance* const NodeInstance = FindNodeInstanceMutable(InNode);
+	if (!ensureAlwaysMsgf(NodeInstance, TEXT("[%.*s] FinishNode called on a node that no longer has node instance data!"), GetContextName().Len(), GetContextName().GetData()))
+	{
+		return false;
+	}
 
-	const bool bWasExecuting = ensureAlwaysMsgf(NodeInstance.Status == EDataLinkNodeStatus::Executing
-		, TEXT("[%s] FinishNode called on a node that was not executing!"), GetContextName().GetData());
+	const bool bWasExecuting = NodeInstance->Status == EDataLinkNodeStatus::Executing;
+	if (!ensureAlwaysMsgf(bWasExecuting, TEXT("[%.*s] FinishNode called on a node that was not executing!"), GetContextName().Len(), GetContextName().GetData()))
+	{
+		return false;
+	}
 
-	NodeInstance.Status = EDataLinkNodeStatus::Finished;
-	return bWasExecuting;
+	NodeInstance->Status = EDataLinkNodeStatus::Finished;
+	return true;
 }
 
 void FDataLinkExecutor::Finish(EDataLinkExecutionResult InStatus, FConstStructView InOutputData)
diff --git a/Engine/Plugins/Experimental/DataLink/Source/DataLink/Public/DataLinkExecutor.h b/Engine/Plugins/Experimental/DataLink/Source/DataLink/Public/DataLinkExecutor.h
--- a/Engine/Plugins/Experimental/DataLink/Source/DataLink/Public/DataLinkExecutor.h
+++ b/Engine/Plugins/Experimental/DataLink/Source/DataLink/Public/DataLinkExecutor.h
@@ -44,11 +44,15 @@ public:
 	/** Called to start the graph execution */
 	DATALINK_API void Run();
 
-	/** Retrieves the Node Instance Data of this execution for the given Node */
+	/** Retrieves the node instance data of this execution for the given node */
 	DATALINK_API const FDataLinkNodeInstance& GetNodeInstance(const UDataLinkNode* InNode) const;
 
+	/** Retrieves the node instance data (mutable) of this execution for the given node */
 	DATALINK_API FDataLinkNodeInstance& GetNodeInstanceMutable(const UDataLinkNode* InNode);
 
+	/** Finds the node instance data (mutable) of this execution for the given node, or null if already removed */
+	DATALINK_API FDataLinkNodeInstance* FindNodeInstanceMutable(const UDataLinkNode* InNode);
+
 	/**
 	 * Called when a node has finished executing successfully
 	 * If it's the output node that has finished executing, it will finish the entire graph execution
diff --git a/Engine/Plugins/Experimental/DataLink/Source/DataLinkHttp/Private/DataLinkHttpSource.cpp b/Engine/Plugins/Experimental/DataLink/Source/DataLinkHttp/Private/DataLinkHttpSource.cpp
--- a/Engine/Plugins/Experimental/DataLink/Source/DataLinkHttp/Private/DataLinkHttpSource.cpp
+++ b/Engine/Plugins/Experimental/DataLink/Source/DataLinkHttp/Private/DataLinkHttpSource.cpp
@@ -76,11 +76,9 @@ EDataLinkExecutionReply UDataLinkHttpSource::OnExecute(FDataLinkExecutor& InExec
 			}
 		});
 
-	const bool bRequestProcessed = HttpRequest->ProcessRequest();
-
-	return bRequestProcessed
-		? EDataLinkExecutionReply::Handled
-		: EDataLinkExecutionReply::Unhandled;
+	// If process request fails, it will be handled within 'OnProcessRequestComplete' lambda above
+	HttpRequest->ProcessRequest();
+	return EDataLinkExecutionReply::Handled;
 }
 
 #undef LOCTEXT_NAMESPACE
