Motion Design: Changed input data in FDataLinkInstance from being just an array of instanced structs to being an array of FDataLinkInputData elements. These contain FInstancedStruct for the data as well as additional information like input display name (currently only used for customization).

diff --git a/Engine/Plugins/Experimental/DataLink/Source/DataLink/Private/DataLinkExecutor.cpp b/Engine/Plugins/Experimental/DataLink/Source/DataLink/Private/DataLinkExecutor.cpp
--- a/Engine/Plugins/Experimental/DataLink/Source/DataLink/Private/DataLinkExecutor.cpp
+++ b/Engine/Plugins/Experimental/DataLink/Source/DataLink/Private/DataLinkExecutor.cpp
@@ -194,14 +194,14 @@ bool FDataLinkExecutor::ValidateRun()
 		{
 			const int32 CurrentPinIndex = PinIndex++;
 
-			if (InPinReference.Pin->Struct && InPinReference.Pin->Struct != Instance.InputData[CurrentPinIndex].GetScriptStruct())
+			if (InPinReference.Pin->Struct && InPinReference.Pin->Struct != Instance.InputData[CurrentPinIndex].Payload.GetScriptStruct())
 			{
 				UE_LOG(LogDataLink, Error, TEXT("[%s] Input pin (%s': input struct '%s', owner '%s') is not compatible with input data view '%s'")
 					, GetContextName().GetData()
 					, *InPinReference.Pin->Name.ToString()
 					, *InPinReference.Pin->Struct->GetName()
 					, *InPinReference.OwningNode->GetName()
-					, *GetNameSafe(Instance.InputData[CurrentPinIndex].GetScriptStruct()));
+					, *GetNameSafe(Instance.InputData[CurrentPinIndex].Payload.GetScriptStruct()));
 				return false;
 			}
 
@@ -247,7 +247,7 @@ bool FDataLinkExecutor::ExecuteEntryNodes()
 	const TArray<FDataLinkPinReference> InputPinsToExecute = Instance.DataLinkGraph->GetInputPins();
 
 	uint16 NodesExecuted = 0;
-	if (!ExecuteInputPins(InputPinsToExecute, TArray<FConstStructView>(Instance.InputData), &NodesExecuted))
+	if (!ExecuteInputPins(InputPinsToExecute, Instance.GetInputDataViews(), &NodesExecuted))
 	{
 		return false;
 	}
diff --git a/Engine/Plugins/Experimental/DataLink/Source/DataLink/Private/DataLinkInstance.cpp b/Engine/Plugins/Experimental/DataLink/Source/DataLink/Private/DataLinkInstance.cpp
new file mode 100644
--- /dev/null
+++ b/Engine/Plugins/Experimental/DataLink/Source/DataLink/Private/DataLinkInstance.cpp
@@ -0,0 +1,27 @@
+// Copyright Epic Games, Inc. All Rights Reserved.
+
+#include "DataLinkInstance.h"
+#include "StructUtils/StructView.h"
+
+bool FDataLinkInputData::SerializeFromMismatchedTag(const FPropertyTag& InPropertyTag, FStructuredArchive::FSlot InSlot)
+{
+	if (InPropertyTag.GetType().IsStruct(FInstancedStruct::StaticStruct()->GetFName()))
+	{
+		FInstancedStruct::StaticStruct()->SerializeItem(InSlot, &Payload, nullptr);
+		return true;
+	}
+	return false;
+}
+
+TArray<FConstStructView> FDataLinkInstance::GetInputDataViews() const
+{
+	TArray<FConstStructView> StructViews;
+	StructViews.Reserve(InputData.Num());
+
+	for (const FDataLinkInputData& InputDataEntry : InputData)
+	{
+		StructViews.Emplace(InputDataEntry.Payload);
+	}
+
+	return StructViews;
+}
diff --git a/Engine/Plugins/Experimental/DataLink/Source/DataLink/Private/DataLinkUtils.cpp b/Engine/Plugins/Experimental/DataLink/Source/DataLink/Private/DataLinkUtils.cpp
--- a/Engine/Plugins/Experimental/DataLink/Source/DataLink/Private/DataLinkUtils.cpp
+++ b/Engine/Plugins/Experimental/DataLink/Source/DataLink/Private/DataLinkUtils.cpp
@@ -1,7 +1,11 @@
 // Copyright Epic Games, Inc. All Rights Reserved.
 
 #include "DataLinkUtils.h"
+#include "DataLinkGraph.h"
+#include "DataLinkInstance.h"
 #include "DataLinkLog.h"
+#include "DataLinkPin.h"
+#include "DataLinkPinReference.h"
 #include "DataLinkSinkObject.h"
 #include "Engine/Engine.h"
 #include "IDataLinkSinkProvider.h"
@@ -9,7 +13,10 @@
 #include "UObject/Object.h"
 #include "UObject/Package.h"
 
-bool UE::DataLink::ReplaceObject(UObject*& InOutObject, UObject* InOuter, UClass* InClass)
+namespace UE::DataLink
+{
+
+bool ReplaceObject(UObject*& InOutObject, UObject* InOuter, UClass* InClass)
 {
 	// Check if the object class already matches the new class
 	if (InOutObject && InOutObject->GetClass() == InClass)
@@ -55,7 +62,7 @@ bool UE::DataLink::ReplaceObject(UObject*& InOutObject, UObject* InOuter, UClass
 	return true;
 }
 
-TSharedPtr<FDataLinkSink> UE::DataLink::TryGetSink(TScriptInterface<IDataLinkSinkProvider> InSinkProvider)
+TSharedPtr<FDataLinkSink> TryGetSink(TScriptInterface<IDataLinkSinkProvider> InSinkProvider)
 {
 	if (IDataLinkSinkProvider* SinkProvider = InSinkProvider.GetInterface())
 	{
@@ -82,3 +89,34 @@ TSharedPtr<FDataLinkSink> UE::DataLink::TryGetSink(TScriptInterface<IDataLinkSin
 
 	return nullptr;
 }
+
+void SetInputData(UDataLinkGraph* InGraph, TArray<FDataLinkInputData>& OutInputData)
+{
+	if (!InGraph)
+	{
+		OutInputData.Reset();
+		return;
+	}
+
+	const int32 InputPinCount = InGraph->GetInputPinCount();
+	OutInputData.SetNum(InputPinCount);
+
+	int32 Index = 0;
+	InGraph->ForEachInputPin(
+		[&OutInputData, &Index](const FDataLinkPinReference& InPinReference)->bool
+		{
+			FDataLinkInputData& InputEntry = OutInputData[Index];
+			InputEntry.InputDisplayName = InPinReference.Pin->GetDisplayName();
+
+			if (InputEntry.Payload.GetScriptStruct() != InPinReference.Pin->Struct)
+			{
+				// Input Struct can be null, FInstancedStruct will just be reset
+				InputEntry.Payload.InitializeAs(InPinReference.Pin->Struct);
+			}
+
+			++Index;
+			return true;
+		});
+}
+
+} // UE::DataLink
diff --git a/Engine/Plugins/Experimental/DataLink/Source/DataLink/Public/DataLinkInstance.h b/Engine/Plugins/Experimental/DataLink/Source/DataLink/Public/DataLinkInstance.h
--- a/Engine/Plugins/Experimental/DataLink/Source/DataLink/Public/DataLinkInstance.h
+++ b/Engine/Plugins/Experimental/DataLink/Source/DataLink/Public/DataLinkInstance.h
@@ -6,19 +6,52 @@
 #include "StructUtils/InstancedStruct.h"
 #include "DataLinkInstance.generated.h"
 
+#define UE_API DATALINK_API
+
 class UDataLinkGraph;
 
+/** Describes an Input Entry to the Data Link */
+USTRUCT(BlueprintType)
+struct FDataLinkInputData
+{
+	GENERATED_BODY()
+
+	bool SerializeFromMismatchedTag(const FPropertyTag& InPropertyTag, FStructuredArchive::FSlot InSlot);
+
+	/** The display name of the input data */
+	UPROPERTY(VisibleAnywhere, Category="Data Link")
+	FText InputDisplayName;
+
+	/** The input data to feed into the data link graph */
+	UPROPERTY(EditAnywhere, EditFixedSize, BlueprintReadWrite, Category="Data Link")
+	FInstancedStruct Payload;
+};
+
+template<>
+struct TStructOpsTypeTraits<FDataLinkInputData> : public TStructOpsTypeTraitsBase2<FDataLinkInputData>
+{
+	enum
+	{
+		WithStructuredSerializeFromMismatchedTag = true,
+	};
+};
+
 /** Instance of a data link to be executed */
 USTRUCT(BlueprintType)
 struct FDataLinkInstance
 {
 	GENERATED_BODY()
 
+	/** Converts the Input Data into an array of Const Struct Views */
+	UE_API TArray<FConstStructView> GetInputDataViews() const;
+
 	/** The data link graph to execute */
 	UPROPERTY(EditAnywhere, BlueprintReadWrite, Category="Data Link")
 	TObjectPtr<UDataLinkGraph> DataLinkGraph;
 
 	/** The initial input data to feed into the data link graph */
 	UPROPERTY(EditAnywhere, EditFixedSize, BlueprintReadWrite, Category="Data Link")
-	TArray<FInstancedStruct> InputData;
+	TArray<FDataLinkInputData> InputData;
 };
+
+#undef UE_API
diff --git a/Engine/Plugins/Experimental/DataLink/Source/DataLink/Public/DataLinkUtils.h b/Engine/Plugins/Experimental/DataLink/Source/DataLink/Public/DataLinkUtils.h
--- a/Engine/Plugins/Experimental/DataLink/Source/DataLink/Public/DataLinkUtils.h
+++ b/Engine/Plugins/Experimental/DataLink/Source/DataLink/Public/DataLinkUtils.h
@@ -2,15 +2,20 @@
 
 #pragma once
 
+#include "Containers/ContainersFwd.h"
 #include "HAL/Platform.h"
 #include "Templates/SharedPointerFwd.h"
 
 class IDataLinkSinkProvider;
 class UClass;
+class UDataLinkGraph;
 class UObject;
+struct FDataLinkInputData;
 struct FDataLinkSink;
 template <typename T> class TScriptInterface;
 
+#define UE_API DATALINK_API
+
 namespace UE::DataLink
 {
 	/**
@@ -20,12 +25,21 @@ namespace UE::DataLink
 	 * @param InClass new class of the new object that will replace the older one
 	 * @return true if the operation took place, false otherwise
 	 */
-	DATALINK_API bool ReplaceObject(UObject*& InOutObject, UObject* InOuter, UClass* InClass);
+	UE_API bool ReplaceObject(UObject*& InOutObject, UObject* InOuter, UClass* InClass);
 
 	/**
 	 * Attempts to get the underlying Sink from the given Sink Provider
 	 * @param InSinkProvider the sink provider interface
 	 * @return the sink if found
 	 */
-	DATALINK_API TSharedPtr<FDataLinkSink> TryGetSink(TScriptInterface<IDataLinkSinkProvider> InSinkProvider);
+	UE_API TSharedPtr<FDataLinkSink> TryGetSink(TScriptInterface<IDataLinkSinkProvider> InSinkProvider);
+
+	/**
+	 * Sets the input data to match the current graph's inputs
+	 * @param InGraph the data link graph to get the input pins to update from
+	 * @param OutInputData the input data to update
+	 */
+	UE_API void SetInputData(UDataLinkGraph* InGraph, TArray<FDataLinkInputData>& OutInputData);
 }
+
+#undef UE_API
diff --git a/Engine/Plugins/Experimental/DataLink/Source/DataLinkEditor/Private/DetailsView/DataLinkInputDataNodeBuilder.cpp b/Engine/Plugins/Experimental/DataLink/Source/DataLinkEditor/Private/DetailsView/DataLinkInputDataNodeBuilder.cpp
new file mode 100644
--- /dev/null
+++ b/Engine/Plugins/Experimental/DataLink/Source/DataLinkEditor/Private/DetailsView/DataLinkInputDataNodeBuilder.cpp
@@ -0,0 +1,81 @@
+// Copyright Epic Games, Inc. All Rights Reserved.
+
+#include "DataLinkInputDataNodeBuilder.h"
+#include "DataLinkInstance.h"
+#include "DetailLayoutBuilder.h"
+#include "DetailWidgetRow.h"
+#include "IDetailChildrenBuilder.h"
+#include "Widgets/SBoxPanel.h"
+#include "Widgets/Text/STextBlock.h"
+
+FDataLinkInputDataNodeBuilder::FDataLinkInputDataNodeBuilder(const TSharedRef<IPropertyHandle>& InInputDataElementHandle)
+	: FInstancedStructDataDetails(InInputDataElementHandle->GetChildHandle(GET_MEMBER_NAME_CHECKED(FDataLinkInputData, Payload)))
+	, InputDataElementHandle(InInputDataElementHandle)
+	, InputNameHandle(InInputDataElementHandle->GetChildHandle(GET_MEMBER_NAME_CHECKED(FDataLinkInputData, InputDisplayName)).ToSharedRef())
+{
+}
+
+void FDataLinkInputDataNodeBuilder::GenerateHeaderRowContent(FDetailWidgetRow& InNodeRow)
+{
+	FInstancedStructDataDetails::GenerateHeaderRowContent(InNodeRow);
+
+	FText InputName;
+	InputNameHandle->GetValueAsDisplayText(InputName);
+
+	InNodeRow
+		.ShouldAutoExpand(true)
+		.NameContent()
+		[
+			SNew(SHorizontalBox)
+			+ SHorizontalBox::Slot()
+			.HAlign(HAlign_Left)
+			.VAlign(VAlign_Center)
+			.AutoWidth()
+			[
+				SNew(STextBlock)
+				.Text(INVTEXT("["))
+				.Font(IDetailLayoutBuilder::GetDetailFont())
+				.ColorAndOpacity(FSlateColor::UseSubduedForeground())
+			]
+			+ SHorizontalBox::Slot()
+			.HAlign(HAlign_Left)
+			.VAlign(VAlign_Center)
+			.Padding(3, 0, 0, 0)
+			.AutoWidth()
+			[
+				SNew(STextBlock)
+				.Text(FText::AsNumber(InputDataElementHandle->GetArrayIndex()))
+				.Font(IDetailLayoutBuilder::GetDetailFont())
+			]
+			+ SHorizontalBox::Slot()
+			.HAlign(HAlign_Left)
+			.VAlign(VAlign_Center)
+			.Padding(3, 0, 0, 0)
+			.AutoWidth()
+			[
+				SNew(STextBlock)
+				.Text(INVTEXT("]"))
+				.Font(IDetailLayoutBuilder::GetDetailFont())
+				.ColorAndOpacity(FSlateColor::UseSubduedForeground())
+			]
+			+ SHorizontalBox::Slot()
+			.HAlign(HAlign_Left)
+			.VAlign(VAlign_Center)
+			.Padding(5, 0, 0, 0)
+			.AutoWidth()
+			[
+				SNew(STextBlock)
+				.Text(InputName)
+				.Font(IDetailLayoutBuilder::GetDetailFont())
+			]
+		];
+}
+
+void FDataLinkInputDataNodeBuilder::Tick(float InDeltaTime)
+{
+	// TODO: Logic in FInstancedStructDataDetails::Tick calls EnumerateConstRawData which does not check that the underlying handle still has a valid property node
+	if (InputDataElementHandle->IsValidHandle())
+	{
+		FInstancedStructDataDetails::Tick(InDeltaTime);
+	}
+}
diff --git a/Engine/Plugins/Experimental/DataLink/Source/DataLinkEditor/Private/DetailsView/DataLinkInputDataNodeBuilder.h b/Engine/Plugins/Experimental/DataLink/Source/DataLinkEditor/Private/DetailsView/DataLinkInputDataNodeBuilder.h
new file mode 100644
--- /dev/null
+++ b/Engine/Plugins/Experimental/DataLink/Source/DataLinkEditor/Private/DetailsView/DataLinkInputDataNodeBuilder.h
@@ -0,0 +1,20 @@
+// Copyright Epic Games, Inc. All Rights Reserved.
+
+#pragma once
+
+#include "InstancedStructDetails.h"
+
+class FDataLinkInputDataNodeBuilder : public FInstancedStructDataDetails
+{
+public:
+	explicit FDataLinkInputDataNodeBuilder(const TSharedRef<IPropertyHandle>& InInputDataElementHandle);
+
+	//~ Begin IDetailCustomNodeBuilder
+	virtual void GenerateHeaderRowContent(FDetailWidgetRow& InNodeRow) override;
+	virtual void Tick(float InDeltaTime) override;
+	//~ End IDetailCustomNodeBuilder
+
+private:
+	TSharedRef<IPropertyHandle> InputDataElementHandle;
+	TSharedRef<IPropertyHandle> InputNameHandle;
+};
diff --git a/Engine/Plugins/Experimental/DataLink/Source/DataLinkEditor/Private/DetailsView/DataLinkInstanceCustomization.cpp b/Engine/Plugins/Experimental/DataLink/Source/DataLinkEditor/Private/DetailsView/DataLinkInstanceCustomization.cpp
--- a/Engine/Plugins/Experimental/DataLink/Source/DataLinkEditor/Private/DetailsView/DataLinkInstanceCustomization.cpp
+++ b/Engine/Plugins/Experimental/DataLink/Source/DataLinkEditor/Private/DetailsView/DataLinkInstanceCustomization.cpp
@@ -2,10 +2,11 @@
 
 #include "DataLinkInstanceCustomization.h"
 #include "DataLinkGraph.h"
+#include "DataLinkInputDataNodeBuilder.h"
 #include "DataLinkInstance.h"
-#include "DataLinkInstancedStructNodeBuilder.h"
 #include "DataLinkPin.h"
 #include "DataLinkPinReference.h"
+#include "DataLinkUtils.h"
 #include "DetailWidgetRow.h"
 #include "IDetailChildrenBuilder.h"
 #include "PropertyCustomizationHelpers.h"
@@ -61,13 +62,9 @@ void FDataLinkInstanceCustomization::CustomizeChildren(TSharedRef<IPropertyHandl
 		, /*DisplayElementNum*/false);
 
 	ArrayBuilder->OnGenerateArrayElementWidget(FOnGenerateArrayElementWidget::CreateLambda(
-		[&InputDisplayNames = InputDisplayNames](TSharedRef<IPropertyHandle> InElementHandle, int32 InIndex, IDetailChildrenBuilder& InChildBuilder)
+		[](TSharedRef<IPropertyHandle> InElementHandle, int32 InIndex, IDetailChildrenBuilder& InChildBuilder)
 		{
-			if (InputDisplayNames.IsValidIndex(InIndex))
-			{
-				InElementHandle->SetPropertyDisplayName(InputDisplayNames[InIndex]);
-			}
-			InChildBuilder.AddCustomBuilder(MakeShared<FDataLinkInstancedStructNodeBuilder>(InElementHandle));
+			InChildBuilder.AddCustomBuilder(MakeShared<FDataLinkInputDataNodeBuilder>(InElementHandle));
 		}));
 
 	InChildBuilder.AddCustomBuilder(ArrayBuilder);
@@ -91,12 +88,12 @@ UDataLinkGraph* FDataLinkInstanceCustomization::GetDataLinkGraph() const
 	return nullptr;
 }
 
-TArray<FInstancedStruct>* FDataLinkInstanceCustomization::GetInputData() const
+TArray<FDataLinkInputData>* FDataLinkInstanceCustomization::GetInputData() const
 {
 	void* RawData = nullptr;
 	if (InputDataHandle.IsValid() && InputDataHandle->GetValueData(RawData) == FPropertyAccess::Success)
 	{
-		return static_cast<TArray<FInstancedStruct>*>(RawData);
+		return static_cast<TArray<FDataLinkInputData>*>(RawData);
 	}
 	return nullptr;
 }
@@ -115,38 +112,8 @@ void FDataLinkInstanceCustomization::OnGraphChanged()
 
 void FDataLinkInstanceCustomization::UpdateInputData()
 {
-	TArray<FInstancedStruct>* InputData = GetInputData();
-	if (!InputData)
+	if (TArray<FDataLinkInputData>* InputData = GetInputData())
 	{
-		return;
-	}
-
-	UDataLinkGraph* DataLinkGraph = GetDataLinkGraph();
-	if (!DataLinkGraph)
-	{
-		InputDisplayNames.Reset();
-		InputData->Reset();
-		return;
+		UE::DataLink::SetInputData(GetDataLinkGraph(), *InputData);
 	}
-
-	const int32 InputPinCount = DataLinkGraph->GetInputPinCount();
-	InputDisplayNames.SetNum(InputPinCount);
-	InputData->SetNum(InputPinCount);
-
-	int32 Index = 0;
-	DataLinkGraph->ForEachInputPin(
-		[InputData, &Index, &InputDisplayNames = InputDisplayNames](const FDataLinkPinReference& InPinReference)->bool
-		{
-			InputDisplayNames[Index] = InPinReference.Pin->GetDisplayName();
-
-			FInstancedStruct& InstanceData = (*InputData)[Index];
-			if (InstanceData.GetScriptStruct() != InPinReference.Pin->Struct)
-			{
-				// Input Struct can be null, FInstancedStruct will just be reset
-				InstanceData.InitializeAs(InPinReference.Pin->Struct);
-			}
-
-			++Index;
-			return true;
-		});
 }
diff --git a/Engine/Plugins/Experimental/DataLink/Source/DataLinkEditor/Private/DetailsView/DataLinkInstanceCustomization.h b/Engine/Plugins/Experimental/DataLink/Source/DataLinkEditor/Private/DetailsView/DataLinkInstanceCustomization.h
--- a/Engine/Plugins/Experimental/DataLink/Source/DataLinkEditor/Private/DetailsView/DataLinkInstanceCustomization.h
+++ b/Engine/Plugins/Experimental/DataLink/Source/DataLinkEditor/Private/DetailsView/DataLinkInstanceCustomization.h
@@ -5,6 +5,7 @@
 #include "IPropertyTypeCustomization.h"
 
 class UDataLinkGraph;
+struct FDataLinkInputData;
 struct FInstancedStruct;
 
 class FDataLinkInstanceCustomization : public IPropertyTypeCustomization
@@ -29,7 +30,7 @@ protected:
 
 	UDataLinkGraph* GetDataLinkGraph() const;
 
-	TArray<FInstancedStruct>* GetInputData() const;
+	TArray<FDataLinkInputData>* GetInputData() const;
 
 	void OnGraphChanged();
 	void UpdateInputData();
@@ -37,7 +38,5 @@ protected:
 	TSharedPtr<IPropertyHandle> DataLinkGraphHandle;
 	TSharedPtr<IPropertyHandle> InputDataHandle;
 
-	TArray<FText> InputDisplayNames;
-
 	bool bGenerateHeader;
 };
diff --git a/Engine/Plugins/Experimental/DataLink/Source/DataLinkEditor/Private/DetailsView/DataLinkInstancedStructNodeBuilder.cpp b/Engine/Plugins/Experimental/DataLink/Source/DataLinkEditor/Private/DetailsView/DataLinkInstancedStructNodeBuilder.cpp
deleted file mode 100644
diff --git a/Engine/Plugins/Experimental/DataLink/Source/DataLinkEditor/Private/DetailsView/DataLinkInstancedStructNodeBuilder.h b/Engine/Plugins/Experimental/DataLink/Source/DataLinkEditor/Private/DetailsView/DataLinkInstancedStructNodeBuilder.h
deleted file mode 100644
