Motion Design: updated input data to have a customization + renamed payload to 'data'

diff --git a/Engine/Plugins/Experimental/DataLink/Source/DataLink/Private/DataLinkExecutor.cpp b/Engine/Plugins/Experimental/DataLink/Source/DataLink/Private/DataLinkExecutor.cpp
--- a/Engine/Plugins/Experimental/DataLink/Source/DataLink/Private/DataLinkExecutor.cpp
+++ b/Engine/Plugins/Experimental/DataLink/Source/DataLink/Private/DataLinkExecutor.cpp
@@ -194,14 +194,14 @@ bool FDataLinkExecutor::ValidateRun()
 		{
 			const int32 CurrentPinIndex = PinIndex++;
 
-			if (InPinReference.Pin->Struct && InPinReference.Pin->Struct != Instance.InputData[CurrentPinIndex].Payload.GetScriptStruct())
+			if (InPinReference.Pin->Struct && InPinReference.Pin->Struct != Instance.InputData[CurrentPinIndex].Data.GetScriptStruct())
 			{
 				UE_LOG(LogDataLink, Error, TEXT("[%s] Input pin (%s': input struct '%s', owner '%s') is not compatible with input data view '%s'")
 					, GetContextName().GetData()
 					, *InPinReference.Pin->Name.ToString()
 					, *InPinReference.Pin->Struct->GetName()
 					, *InPinReference.OwningNode->GetName()
-					, *GetNameSafe(Instance.InputData[CurrentPinIndex].Payload.GetScriptStruct()));
+					, *GetNameSafe(Instance.InputData[CurrentPinIndex].Data.GetScriptStruct()));
 				return false;
 			}
 
diff --git a/Engine/Plugins/Experimental/DataLink/Source/DataLink/Private/DataLinkInstance.cpp b/Engine/Plugins/Experimental/DataLink/Source/DataLink/Private/DataLinkInstance.cpp
--- a/Engine/Plugins/Experimental/DataLink/Source/DataLink/Private/DataLinkInstance.cpp
+++ b/Engine/Plugins/Experimental/DataLink/Source/DataLink/Private/DataLinkInstance.cpp
@@ -7,7 +7,7 @@ bool FDataLinkInputData::SerializeFromMismatchedTag(const FPropertyTag& InProper
 {
 	if (InPropertyTag.GetType().IsStruct(FInstancedStruct::StaticStruct()->GetFName()))
 	{
-		FInstancedStruct::StaticStruct()->SerializeItem(InSlot, &Payload, nullptr);
+		FInstancedStruct::StaticStruct()->SerializeItem(InSlot, &Data, nullptr);
 		return true;
 	}
 	return false;
@@ -20,7 +20,7 @@ TArray<FConstStructView> FDataLinkInstance::GetInputDataViews() const
 
 	for (const FDataLinkInputData& InputDataEntry : InputData)
 	{
-		StructViews.Emplace(InputDataEntry.Payload);
+		StructViews.Emplace(InputDataEntry.Data);
 	}
 
 	return StructViews;
diff --git a/Engine/Plugins/Experimental/DataLink/Source/DataLink/Private/DataLinkUtils.cpp b/Engine/Plugins/Experimental/DataLink/Source/DataLink/Private/DataLinkUtils.cpp
--- a/Engine/Plugins/Experimental/DataLink/Source/DataLink/Private/DataLinkUtils.cpp
+++ b/Engine/Plugins/Experimental/DataLink/Source/DataLink/Private/DataLinkUtils.cpp
@@ -106,12 +106,12 @@ void SetInputData(UDataLinkGraph* InGraph, TArray<FDataLinkInputData>& OutInputD
 		[&OutInputData, &Index](const FDataLinkPinReference& InPinReference)->bool
 		{
 			FDataLinkInputData& InputEntry = OutInputData[Index];
-			InputEntry.InputDisplayName = InPinReference.Pin->GetDisplayName();
+			InputEntry.DisplayName = InPinReference.Pin->GetDisplayName();
 
-			if (InputEntry.Payload.GetScriptStruct() != InPinReference.Pin->Struct)
+			if (InputEntry.Data.GetScriptStruct() != InPinReference.Pin->Struct)
 			{
 				// Input Struct can be null, FInstancedStruct will just be reset
-				InputEntry.Payload.InitializeAs(InPinReference.Pin->Struct);
+				InputEntry.Data.InitializeAs(InPinReference.Pin->Struct);
 			}
 
 			++Index;
diff --git a/Engine/Plugins/Experimental/DataLink/Source/DataLink/Public/DataLinkInstance.h b/Engine/Plugins/Experimental/DataLink/Source/DataLink/Public/DataLinkInstance.h
--- a/Engine/Plugins/Experimental/DataLink/Source/DataLink/Public/DataLinkInstance.h
+++ b/Engine/Plugins/Experimental/DataLink/Source/DataLink/Public/DataLinkInstance.h
@@ -20,11 +20,11 @@ struct FDataLinkInputData
 
 	/** The display name of the input data */
 	UPROPERTY(VisibleAnywhere, Category="Data Link")
-	FText InputDisplayName;
+	FText DisplayName;
 
 	/** The input data to feed into the data link graph */
-	UPROPERTY(EditAnywhere, EditFixedSize, BlueprintReadWrite, Category="Data Link")
-	FInstancedStruct Payload;
+	UPROPERTY(EditAnywhere, BlueprintReadWrite, Category="Data Link", meta=(StructTypeConst))
+	FInstancedStruct Data;
 };
 
 template<>
@@ -50,7 +50,7 @@ struct FDataLinkInstance
 	TObjectPtr<UDataLinkGraph> DataLinkGraph;
 
 	/** The initial input data to feed into the data link graph */
-	UPROPERTY(EditAnywhere, EditFixedSize, BlueprintReadWrite, Category="Data Link")
+	UPROPERTY(EditAnywhere, EditFixedSize, BlueprintReadWrite, Category="Data Link", meta=(EditFixedOrder))
 	TArray<FDataLinkInputData> InputData;
 };
 
diff --git a/Engine/Plugins/Experimental/DataLink/Source/DataLinkEditor/Private/DataLinkEditorModule.cpp b/Engine/Plugins/Experimental/DataLink/Source/DataLinkEditor/Private/DataLinkEditorModule.cpp
--- a/Engine/Plugins/Experimental/DataLink/Source/DataLinkEditor/Private/DataLinkEditorModule.cpp
+++ b/Engine/Plugins/Experimental/DataLink/Source/DataLinkEditor/Private/DataLinkEditorModule.cpp
@@ -4,6 +4,7 @@
 #include "DataLinkEditorLog.h"
 #include "DataLinkGraphCommands.h"
 #include "DataLinkInstance.h"
+#include "DetailsView/DataLinkInputDataCustomization.h"
 #include "DetailsView/DataLinkInstanceCustomization.h"
 #include "Modules/ModuleManager.h"
 #include "PropertyEditorModule.h"
@@ -28,6 +29,9 @@ void FDataLinkEditorModule::RegisterDetailCustomizations()
 {
 	FPropertyEditorModule& PropertyEditorModule = FModuleManager::LoadModuleChecked<FPropertyEditorModule>("PropertyEditor");
 
+	PropertyEditorModule.RegisterCustomPropertyTypeLayout(CustomizedTypes.Add_GetRef(FDataLinkInputData::StaticStruct()->GetFName())
+		, FOnGetPropertyTypeCustomizationInstance::CreateStatic(&FDataLinkInputDataCustomization::MakeInstance));
+
 	PropertyEditorModule.RegisterCustomPropertyTypeLayout(CustomizedTypes.Add_GetRef(FDataLinkInstance::StaticStruct()->GetFName())
 		, FOnGetPropertyTypeCustomizationInstance::CreateStatic(&FDataLinkInstanceCustomization::MakeInstance, /*bGenerateHeader*/true));
 }
diff --git a/Engine/Plugins/Experimental/DataLink/Source/DataLinkEditor/Private/DetailsView/DataLinkInputDataCustomization.cpp b/Engine/Plugins/Experimental/DataLink/Source/DataLinkEditor/Private/DetailsView/DataLinkInputDataCustomization.cpp
new file mode 100644
--- /dev/null
+++ b/Engine/Plugins/Experimental/DataLink/Source/DataLinkEditor/Private/DetailsView/DataLinkInputDataCustomization.cpp
@@ -0,0 +1,104 @@
+// Copyright Epic Games, Inc. All Rights Reserved.
+
+#include "DataLinkInputDataCustomization.h"
+#include "DataLinkInstance.h"
+#include "DetailLayoutBuilder.h"
+#include "DetailWidgetRow.h"
+#include "IDetailChildrenBuilder.h"
+#include "InstancedStructDetails.h"
+#include "Widgets/SBoxPanel.h"
+#include "Widgets/Text/STextBlock.h"
+
+namespace UE::DataLink::Private
+{
+	TSharedRef<SWidget> MakeArrayIndexWidget(int32 InArrayIndex)
+	{
+		return SNew(SHorizontalBox)
+			+ SHorizontalBox::Slot()
+			.HAlign(HAlign_Left)
+			.VAlign(VAlign_Center)
+			.AutoWidth()
+			[
+				SNew(STextBlock)
+				.Text(INVTEXT("["))
+				.Font(IDetailLayoutBuilder::GetDetailFont())
+				.ColorAndOpacity(FSlateColor::UseSubduedForeground())
+			]
+			+ SHorizontalBox::Slot()
+			.HAlign(HAlign_Left)
+			.VAlign(VAlign_Center)
+			.Padding(3, 0, 0, 0)
+			.AutoWidth()
+			[
+				SNew(STextBlock)
+				.Text(FText::AsNumber(InArrayIndex))
+				.Font(IDetailLayoutBuilder::GetDetailFont())
+			]
+			+ SHorizontalBox::Slot()
+			.HAlign(HAlign_Left)
+			.VAlign(VAlign_Center)
+			.Padding(3, 0, 0, 0)
+			.AutoWidth()
+			[
+				SNew(STextBlock)
+				.Text(INVTEXT("]"))
+				.Font(IDetailLayoutBuilder::GetDetailFont())
+				.ColorAndOpacity(FSlateColor::UseSubduedForeground())
+			];
+	}
+}
+
+void FDataLinkInputDataCustomization::CustomizeHeader(TSharedRef<IPropertyHandle> InPropertyHandle, FDetailWidgetRow& InHeaderRow, IPropertyTypeCustomizationUtils& InCustomizationUtils)
+{
+	DisplayNameHandle = InPropertyHandle->GetChildHandle(GET_MEMBER_NAME_CHECKED(FDataLinkInputData, DisplayName)).ToSharedRef();
+	DisplayNameHandle->MarkHiddenByCustomization();
+
+	TSharedRef<SHorizontalBox> NameBox = SNew(SHorizontalBox);
+
+	const int32 ArrayIndex = InPropertyHandle->GetArrayIndex();
+	if (ArrayIndex != INDEX_NONE)
+	{
+		NameBox->AddSlot()
+			.AutoWidth()
+			.HAlign(HAlign_Left)
+			.VAlign(VAlign_Center)
+			[
+				UE::DataLink::Private::MakeArrayIndexWidget(ArrayIndex)
+			];
+	}
+
+	NameBox->AddSlot()
+		.AutoWidth()
+		.HAlign(HAlign_Left)
+		.VAlign(VAlign_Center)
+		.Padding(5, 0, 0, 0)
+		[
+			SNew(STextBlock)
+			.Text(this, &FDataLinkInputDataCustomization::GetInputDisplayName)
+			.Font(IDetailLayoutBuilder::GetDetailFont())
+		];
+
+	InHeaderRow
+		.ShouldAutoExpand(true)
+		.NameContent()
+		[
+			NameBox
+		];
+}
+
+void FDataLinkInputDataCustomization::CustomizeChildren(TSharedRef<IPropertyHandle> InPropertyHandle, IDetailChildrenBuilder& InChildBuilder, IPropertyTypeCustomizationUtils& InCustomizationUtils)
+{
+	const TSharedRef<IPropertyHandle> DataHandle = InPropertyHandle->GetChildHandle(GET_MEMBER_NAME_CHECKED(FDataLinkInputData, Data)).ToSharedRef();
+	DataHandle->MarkHiddenByCustomization();
+	InChildBuilder.AddCustomBuilder(MakeShared<FInstancedStructDataDetails>(DataHandle));
+}
+
+FText FDataLinkInputDataCustomization::GetInputDisplayName() const
+{
+	FText InputDisplayName;
+	if (DisplayNameHandle.IsValid())
+	{
+		DisplayNameHandle->GetValueAsDisplayText(InputDisplayName);
+	}
+	return InputDisplayName;
+}
diff --git a/Engine/Plugins/Experimental/DataLink/Source/DataLinkEditor/Private/DetailsView/DataLinkInputDataCustomization.h b/Engine/Plugins/Experimental/DataLink/Source/DataLinkEditor/Private/DetailsView/DataLinkInputDataCustomization.h
new file mode 100644
--- /dev/null
+++ b/Engine/Plugins/Experimental/DataLink/Source/DataLinkEditor/Private/DetailsView/DataLinkInputDataCustomization.h
@@ -0,0 +1,27 @@
+// Copyright Epic Games, Inc. All Rights Reserved.
+
+#pragma once
+
+#include "IPropertyTypeCustomization.h"
+
+class FDataLinkInputDataCustomization : public IPropertyTypeCustomization
+{
+public:
+	static TSharedRef<IPropertyTypeCustomization> MakeInstance()
+	{
+		return MakeShared<FDataLinkInputDataCustomization>();
+	}
+
+protected:
+	//~ Begin IPropertyTypeCustomization
+	virtual void CustomizeHeader(TSharedRef<IPropertyHandle> InPropertyHandle, FDetailWidgetRow& InHeaderRow, IPropertyTypeCustomizationUtils& InCustomizationUtils) override;
+	virtual void CustomizeChildren(TSharedRef<IPropertyHandle> InPropertyHandle, IDetailChildrenBuilder& InChildBuilder, IPropertyTypeCustomizationUtils& InCustomizationUtils) override;
+	//~ End IPropertyTypeCustomization
+
+	/** Returns the current value of the input display name property handle */
+	FText GetInputDisplayName() const;
+
+private:
+	/** Handle to the Input Data Display Name */
+	TSharedPtr<IPropertyHandle> DisplayNameHandle;
+};
diff --git a/Engine/Plugins/Experimental/DataLink/Source/DataLinkEditor/Private/DetailsView/DataLinkInputDataNodeBuilder.cpp b/Engine/Plugins/Experimental/DataLink/Source/DataLinkEditor/Private/DetailsView/DataLinkInputDataNodeBuilder.cpp
deleted file mode 100644
diff --git a/Engine/Plugins/Experimental/DataLink/Source/DataLinkEditor/Private/DetailsView/DataLinkInputDataNodeBuilder.h b/Engine/Plugins/Experimental/DataLink/Source/DataLinkEditor/Private/DetailsView/DataLinkInputDataNodeBuilder.h
deleted file mode 100644
diff --git a/Engine/Plugins/Experimental/DataLink/Source/DataLinkEditor/Private/DetailsView/DataLinkInstanceCustomization.cpp b/Engine/Plugins/Experimental/DataLink/Source/DataLinkEditor/Private/DetailsView/DataLinkInstanceCustomization.cpp
--- a/Engine/Plugins/Experimental/DataLink/Source/DataLinkEditor/Private/DetailsView/DataLinkInstanceCustomization.cpp
+++ b/Engine/Plugins/Experimental/DataLink/Source/DataLinkEditor/Private/DetailsView/DataLinkInstanceCustomization.cpp
@@ -2,7 +2,6 @@
 
 #include "DataLinkInstanceCustomization.h"
 #include "DataLinkGraph.h"
-#include "DataLinkInputDataNodeBuilder.h"
 #include "DataLinkInstance.h"
 #include "DataLinkPin.h"
 #include "DataLinkPinReference.h"
@@ -64,7 +63,7 @@ void FDataLinkInstanceCustomization::CustomizeChildren(TSharedRef<IPropertyHandl
 	ArrayBuilder->OnGenerateArrayElementWidget(FOnGenerateArrayElementWidget::CreateLambda(
 		[](TSharedRef<IPropertyHandle> InElementHandle, int32 InIndex, IDetailChildrenBuilder& InChildBuilder)
 		{
-			InChildBuilder.AddCustomBuilder(MakeShared<FDataLinkInputDataNodeBuilder>(InElementHandle));
+			InChildBuilder.AddProperty(InElementHandle);
 		}));
 
 	InChildBuilder.AddCustomBuilder(ArrayBuilder);
