Motion Design: fixed data link LeftChopInline incorrectly used where LeftInline should've been used

diff --git a/Engine/Plugins/Experimental/DataLink/Source/DataLinkJson/Private/DataLinkJsonUtils.cpp b/Engine/Plugins/Experimental/DataLink/Source/DataLinkJson/Private/DataLinkJsonUtils.cpp
--- a/Engine/Plugins/Experimental/DataLink/Source/DataLinkJson/Private/DataLinkJsonUtils.cpp
+++ b/Engine/Plugins/Experimental/DataLink/Source/DataLinkJson/Private/DataLinkJsonUtils.cpp
@@ -101,7 +101,7 @@ TSharedPtr<FJsonValue> UE::DataLinkJson::FindJsonValue(const TSharedRef<FJsonObj
 		if (!ArraySubscripts.IsEmpty())
 		{
 			const int32 FirstSubscriptIndex = ArraySubscripts[0].GetData() - FieldNameView.GetData();
-			FieldNameView.LeftChopInline(FirstSubscriptIndex - Private::LeftArrayDelimiter.Len());
+			FieldNameView.LeftInline(FirstSubscriptIndex - Private::LeftArrayDelimiter.Len());
 		}
 
 		TSharedPtr<FJsonValue> Value = CurrentObject->TryGetField(FieldNameView);
