Motion Design: fixed Data Link Constant with Json Mapping to work in packaged game.

diff --git a/Engine/Plugins/Experimental/DataLink/Source/DataLink/Private/DataLinkExecutor.cpp b/Engine/Plugins/Experimental/DataLink/Source/DataLink/Private/DataLinkExecutor.cpp
--- a/Engine/Plugins/Experimental/DataLink/Source/DataLink/Private/DataLinkExecutor.cpp
+++ b/Engine/Plugins/Experimental/DataLink/Source/DataLink/Private/DataLinkExecutor.cpp
@@ -8,9 +8,17 @@
 #include "DataLinkNode.h"
 #include "DataLinkPinReference.h"
 #include "DataLinkSink.h"
+#include "HAL/IConsoleManager.h"
 #include "Templates/SubclassOf.h"
 #include "UObject/Class.h"
 
+namespace UE::DataLink::Private
+{
+	static TAutoConsoleVariable<bool> CVarDataLinkDebugNodeData(TEXT("DataLink.PrintNodeData")
+		, /*DefaultValue*/false
+		, TEXT("Prints the input and output data structs for each node executed."));
+}
+
 TSharedPtr<FDataLinkExecutor> FDataLinkExecutor::Create(FDataLinkExecutorArguments&& InArgs)
 {
 	TSharedRef<FDataLinkExecutor> Executor = MakeShared<FDataLinkExecutor>(FPrivateToken());
@@ -392,6 +400,22 @@ bool FDataLinkExecutor::ProcessFinishedNode(const UDataLinkNode* InNode)
 		return false;
 	}
 
+#if !NO_LOGGING
+	if (UE::DataLink::Private::CVarDataLinkDebugNodeData.GetValueOnGameThread())
+	{
+		UE_LOG(LogDataLink, Log, TEXT("[%.*s] Node %s Finished Execution."
+			"\n-------------------------------------------------------------"
+			"\nINPUT PINS:\n%s"
+			"\n-------------------------------------------------------------"
+			"\nOUTPUT PINS:\n%s"
+			"\n-------------------------------------------------------------")
+			, GetContextName().Len(), GetContextName().GetData()
+			, *InNode->GetName()
+			, *NodeInstance->GetInputDataViewer().ToDebugString()
+			, *NodeInstance->GetOutputDataViewer().ToDebugString());
+	}
+#endif
+
 	NodeInstance->Status = EDataLinkNodeStatus::Finished;
 	return true;
 }
diff --git a/Engine/Plugins/Experimental/DataLink/Source/DataLink/Private/DataLinkInputDataViewer.cpp b/Engine/Plugins/Experimental/DataLink/Source/DataLink/Private/DataLinkInputDataViewer.cpp
--- a/Engine/Plugins/Experimental/DataLink/Source/DataLink/Private/DataLinkInputDataViewer.cpp
+++ b/Engine/Plugins/Experimental/DataLink/Source/DataLink/Private/DataLinkInputDataViewer.cpp
@@ -2,6 +2,12 @@
 
 #include "DataLinkInputDataViewer.h"
 #include "DataLinkPin.h"
+#include "DataLinkUtils.h"
+
+FString FDataLinkInputDataEntry::ToDebugString() const
+{
+	return Name.ToString() + TEXT(" : ") + UE::DataLink::StructViewToDebugString(DataView);
+}
 
 FDataLinkInputDataViewer::FDataLinkInputDataViewer(TConstArrayView<FDataLinkPin> InInputPins)
 {
@@ -29,6 +35,15 @@ TConstArrayView<FDataLinkInputDataEntry> FDataLinkInputDataViewer::GetDataEntrie
 	return DataEntries;
 }
 
+FString FDataLinkInputDataViewer::ToDebugString() const
+{
+	return FString::JoinBy(DataEntries, TEXT("\n"),
+		[](const FDataLinkInputDataEntry& InDataEntry)
+		{
+			return InDataEntry.ToDebugString();
+		});
+}
+
 bool FDataLinkInputDataViewer::HasInvalidDataEntry() const
 {
 	for (const FDataLinkInputDataEntry& DataEntry : DataEntries)
diff --git a/Engine/Plugins/Experimental/DataLink/Source/DataLink/Private/DataLinkNode.cpp b/Engine/Plugins/Experimental/DataLink/Source/DataLink/Private/DataLinkNode.cpp
--- a/Engine/Plugins/Experimental/DataLink/Source/DataLink/Private/DataLinkNode.cpp
+++ b/Engine/Plugins/Experimental/DataLink/Source/DataLink/Private/DataLinkNode.cpp
@@ -19,6 +19,11 @@ void UDataLinkNode::Execute(FDataLinkExecutor& InExecutor) const
 }
 
 #if WITH_EDITOR
+void UDataLinkNode::FixupNode()
+{
+	OnFixupNode();
+}
+
 void UDataLinkNode::BuildMetadata(FDataLinkNodeMetadata& OutMetadata) const
 {
 	OutMetadata
diff --git a/Engine/Plugins/Experimental/DataLink/Source/DataLink/Private/DataLinkOutputDataViewer.cpp b/Engine/Plugins/Experimental/DataLink/Source/DataLink/Private/DataLinkOutputDataViewer.cpp
--- a/Engine/Plugins/Experimental/DataLink/Source/DataLink/Private/DataLinkOutputDataViewer.cpp
+++ b/Engine/Plugins/Experimental/DataLink/Source/DataLink/Private/DataLinkOutputDataViewer.cpp
@@ -2,6 +2,7 @@
 
 #include "DataLinkOutputDataViewer.h"
 #include "DataLinkPin.h"
+#include "DataLinkUtils.h"
 
 FStructView FDataLinkOutputDataEntry::GetDataView(const UScriptStruct* InDesiredStruct) const
 {
@@ -12,6 +13,11 @@ FStructView FDataLinkOutputDataEntry::GetDataView(const UScriptStruct* InDesired
 	return FStructView(OutputData);
 }
 
+FString FDataLinkOutputDataEntry::ToDebugString() const
+{
+	return Name.ToString() + TEXT(" : ") + UE::DataLink::StructViewToDebugString(OutputData);
+}
+
 FDataLinkOutputDataViewer::FDataLinkOutputDataViewer(TConstArrayView<FDataLinkPin> InOutputPins)
 {
 	DataEntries.Reserve(InOutputPins.Num());
@@ -32,3 +38,12 @@ FStructView FDataLinkOutputDataViewer::Find(FName InOutputName, const UScriptStr
 	}
 	return FStructView();
 }
+
+FString FDataLinkOutputDataViewer::ToDebugString() const
+{
+	return FString::JoinBy(DataEntries, TEXT("\n"),
+		[](const FDataLinkOutputDataEntry& InDataEntry)
+		{
+			return InDataEntry.ToDebugString();
+		});
+}
diff --git a/Engine/Plugins/Experimental/DataLink/Source/DataLink/Private/DataLinkUtils.cpp b/Engine/Plugins/Experimental/DataLink/Source/DataLink/Private/DataLinkUtils.cpp
--- a/Engine/Plugins/Experimental/DataLink/Source/DataLink/Private/DataLinkUtils.cpp
+++ b/Engine/Plugins/Experimental/DataLink/Source/DataLink/Private/DataLinkUtils.cpp
@@ -200,4 +200,17 @@ void SetInputData(UDataLinkGraph* InGraph, TArray<FDataLinkInputData>& OutInputD
 		});
 }
 
+FString StructViewToDebugString(FConstStructView InDataView)
+{
+	if (!InDataView.IsValid())
+	{
+		return FString();
+	}
+
+	FString DebugString;
+	constexpr int32 PortFlags = PPF_PropertyWindow | PPF_BlueprintDebugView;
+	InDataView.GetScriptStruct()->ExportText(DebugString, InDataView.GetMemory(), /*Defaults*/nullptr, /*OwnerObject*/nullptr, PortFlags, /*ExportRootScope*/nullptr);
+	return DebugString;
+}
+
 } // UE::DataLink
diff --git a/Engine/Plugins/Experimental/DataLink/Source/DataLink/Private/Nodes/DataLinkConstant.cpp b/Engine/Plugins/Experimental/DataLink/Source/DataLink/Private/Nodes/DataLinkConstant.cpp
--- a/Engine/Plugins/Experimental/DataLink/Source/DataLink/Private/Nodes/DataLinkConstant.cpp
+++ b/Engine/Plugins/Experimental/DataLink/Source/DataLink/Private/Nodes/DataLinkConstant.cpp
@@ -2,6 +2,7 @@
 
 #include "Nodes/DataLinkConstant.h"
 #include "DataLinkExecutor.h"
+#include "DataLinkLog.h"
 #include "DataLinkNames.h"
 #include "DataLinkNodeInstance.h"
 #include "DataLinkNodeMetadata.h"
@@ -13,6 +14,53 @@ void UDataLinkConstant::SetStruct(const UScriptStruct* InStructType)
 }
 
 #if WITH_EDITOR
+void UDataLinkConstant::OnFixupNode()
+{
+	Super::OnFixupNode();
+
+	if (!Instance.IsValid())
+	{
+		return;
+	}
+
+	// For Instanced Object properties, make sure the object is duplicated and under this node
+	auto InstanceObjectsFunc = [This=this](const FObjectProperty* InObjectProperty, void* InValueAddress)
+		{
+			if (!InObjectProperty || !InObjectProperty->HasAllPropertyFlags(CPF_InstancedReference))
+			{
+				return;
+			}
+
+			// Only consider objects that should be outered to this constant node by verifying that its current outer is also a constant node.
+			UObject* const TemplateObject = InObjectProperty->GetObjectPropertyValue(InValueAddress);
+			if (!TemplateObject || !TemplateObject->GetTypedOuter<UDataLinkConstant>())
+			{
+				return;
+			}
+
+			FObjectDuplicationParameters Parameters = InitStaticDuplicateObjectParams(TemplateObject, This);
+			Parameters.DestName = MakeUniqueObjectName(Parameters.DestOuter, Parameters.SourceObject->GetClass(), Parameters.SourceObject->GetFName());
+			Parameters.FlagMask = RF_AllFlags & ~RF_DefaultSubObject;
+			Parameters.PortFlags |= PPF_DuplicateVerbatim; // Skip resetting text IDs
+
+			UObject* InstanceObject = StaticDuplicateObjectEx(Parameters);
+			InObjectProperty->SetObjectPropertyValue(InValueAddress, InstanceObject);
+		};
+
+	Instance.GetScriptStruct()->Visit(Instance.GetMutableMemory(),
+		[&InstanceObjectsFunc](const FPropertyVisitorContext& InContext)->EPropertyVisitorControlFlow
+		{
+			if (const FObjectProperty* ObjectProperty = CastField<FObjectProperty>(InContext.Path.Top().Property))
+			{
+				InstanceObjectsFunc(ObjectProperty, InContext.Data.PropertyData);
+
+				// Do not iterate into object. Object properties that don't have a UObject parent are the only ones that could need fixing.
+				return EPropertyVisitorControlFlow::StepOver;
+			}
+			return EPropertyVisitorControlFlow::StepInto;
+		});
+}
+
 void UDataLinkConstant::OnBuildMetadata(FDataLinkNodeMetadata& Metadata) const
 {
 	Super::OnBuildMetadata(Metadata);
@@ -44,6 +92,28 @@ void UDataLinkConstant::OnBuildPins(FDataLinkPinBuilder& Inputs, FDataLinkPinBui
 
 EDataLinkExecutionReply UDataLinkConstant::OnExecute(FDataLinkExecutor& InExecutor) const
 {
-	InExecutor.SucceedNode(this, FConstStructView(Instance.GetScriptStruct(), Instance.GetMemory()));
+	if (!Instance.IsValid())
+	{
+		UE_LOG(LogDataLink, Error, TEXT("[%.*s] DataLinkConstant failed. Instance is not valid")
+			, InExecutor.GetContextName().Len(), InExecutor.GetContextName().GetData());
+		return EDataLinkExecutionReply::Unhandled;
+	}
+
+	const FDataLinkNodeInstance& NodeInstance = InExecutor.GetNodeInstance(this);
+
+	const FDataLinkOutputDataViewer& OutputDataViewer = NodeInstance.GetOutputDataViewer();
+
+	const FStructView OutputDataView = OutputDataViewer.Find(UE::DataLink::OutputDefault, Instance.GetScriptStruct());
+	if (Instance.GetScriptStruct() != OutputDataView.GetScriptStruct())
+	{
+		UE_LOG(LogDataLink, Error, TEXT("[%.*s] DataLinkConstant failed. %s does not match output struct %s")
+			, InExecutor.GetContextName().Len(), InExecutor.GetContextName().GetData()
+			, *GetNameSafe(Instance.GetScriptStruct())
+			, *GetNameSafe(OutputDataView.GetScriptStruct()));
+		return EDataLinkExecutionReply::Unhandled;
+	}
+
+	OutputDataView.GetScriptStruct()->CopyScriptStruct(OutputDataView.GetMemory(), Instance.GetMemory());
+	InExecutor.SucceedNode(this, OutputDataView);
 	return EDataLinkExecutionReply::Handled;
 }
diff --git a/Engine/Plugins/Experimental/DataLink/Source/DataLink/Public/DataLinkInputDataViewer.h b/Engine/Plugins/Experimental/DataLink/Source/DataLink/Public/DataLinkInputDataViewer.h
--- a/Engine/Plugins/Experimental/DataLink/Source/DataLink/Public/DataLinkInputDataViewer.h
+++ b/Engine/Plugins/Experimental/DataLink/Source/DataLink/Public/DataLinkInputDataViewer.h
@@ -4,6 +4,8 @@
 
 #include "StructUtils/StructView.h"
 
+#define UE_API DATALINK_API
+
 struct FDataLinkInputDataEntry
 {
 	bool operator==(FName InInputName) const
@@ -11,6 +13,9 @@ struct FDataLinkInputDataEntry
 		return Name == InInputName;
 	}
 
+	/** Converts the input data to a string for debugging purposes */
+	FString ToDebugString() const;
+
 	/** Name of the Input Data */
 	FName Name;
 
@@ -27,7 +32,7 @@ class FDataLinkInputDataViewer
 public:
 	explicit FDataLinkInputDataViewer(TConstArrayView<FDataLinkPin> InInputPins);
 
-	DATALINK_API FConstStructView Find(FName InInputName) const;
+	UE_API FConstStructView Find(FName InInputName) const;
 
 	int32 Num() const
 	{
@@ -42,6 +47,9 @@ public:
 
 	TConstArrayView<FDataLinkInputDataEntry> GetDataEntries() const;
 
+	/** Converts the output data to a string for debugging purposes */
+	UE_API FString ToDebugString() const;
+	
 private:
 	bool HasInvalidDataEntry() const;
 
@@ -49,3 +57,5 @@ private:
 
 	TArray<FDataLinkInputDataEntry> DataEntries;
 };
+
+#undef UE_API
diff --git a/Engine/Plugins/Experimental/DataLink/Source/DataLink/Public/DataLinkNode.h b/Engine/Plugins/Experimental/DataLink/Source/DataLink/Public/DataLinkNode.h
--- a/Engine/Plugins/Experimental/DataLink/Source/DataLink/Public/DataLinkNode.h
+++ b/Engine/Plugins/Experimental/DataLink/Source/DataLink/Public/DataLinkNode.h
@@ -7,6 +7,8 @@
 #include "UObject/Object.h"
 #include "DataLinkNode.generated.h"
 
+#define UE_API DATALINK_API
+
 class FDataLinkExecutor;
 struct FDataLinkPinBuilder;
 
@@ -46,19 +48,29 @@ public:
 	}
 
 #if WITH_EDITOR
-	DATALINK_API void BuildMetadata(FDataLinkNodeMetadata& OutMetadata) const;
+	/** Called when the node has been compiled to fix data to be valid to use at runtime */
+	UE_API void FixupNode();
+	UE_API void BuildMetadata(FDataLinkNodeMetadata& OutMetadata) const;
 #endif
 
-	DATALINK_API void BuildPins(TArray<FDataLinkPin>& OutInputPins, TArray<FDataLinkPin>& OutOutputPins) const;
+	UE_API void BuildPins(TArray<FDataLinkPin>& OutInputPins, TArray<FDataLinkPin>& OutOutputPins) const;
 
 protected:
 #if WITH_EDITOR
+	virtual void OnFixupNode()
+	{
+	}
+
 	/** Call to retrieve the metadata for this node */
-	virtual void OnBuildMetadata(FDataLinkNodeMetadata& Metadata) const {}
+	virtual void OnBuildMetadata(FDataLinkNodeMetadata& Metadata) const
+	{
+	}
 #endif
 
 	/** Call to setup the pins required by this node */
-	virtual void OnBuildPins(FDataLinkPinBuilder& Inputs, FDataLinkPinBuilder& Outputs) const {}
+	virtual void OnBuildPins(FDataLinkPinBuilder& Inputs, FDataLinkPinBuilder& Outputs) const
+	{
+	}
 
 	/**
 	 * Execution logic of the node
@@ -82,3 +94,5 @@ private:
 	UPROPERTY()
 	TArray<FDataLinkPin> OutputPins;
 };
+
+#undef UE_API
diff --git a/Engine/Plugins/Experimental/DataLink/Source/DataLink/Public/DataLinkOutputDataViewer.h b/Engine/Plugins/Experimental/DataLink/Source/DataLink/Public/DataLinkOutputDataViewer.h
--- a/Engine/Plugins/Experimental/DataLink/Source/DataLink/Public/DataLinkOutputDataViewer.h
+++ b/Engine/Plugins/Experimental/DataLink/Source/DataLink/Public/DataLinkOutputDataViewer.h
@@ -4,6 +4,8 @@
 
 #include "StructUtils/StructView.h"
 
+#define UE_API DATALINK_API
+
 struct FDataLinkPin;
 
 struct FDataLinkOutputDataEntry
@@ -18,6 +20,9 @@ struct FDataLinkOutputDataEntry
 
 	FStructView GetDataView(const UScriptStruct* InDesiredStruct) const;
 
+	/** Converts the output data to a string for debugging purposes */
+	FString ToDebugString() const;
+
 private:
 	/** Output Data instantiated */
 	mutable FInstancedStruct OutputData;
@@ -28,7 +33,7 @@ class FDataLinkOutputDataViewer
 public:
 	explicit FDataLinkOutputDataViewer(TConstArrayView<FDataLinkPin> InOutputPins);
 
-	DATALINK_API FStructView Find(FName InOutputName, const UScriptStruct* InDesiredStruct) const;
+	UE_API FStructView Find(FName InOutputName, const UScriptStruct* InDesiredStruct) const;
 
 	int32 Num() const
 	{
@@ -41,6 +46,10 @@ public:
 		return this->Find(InOutputName, T::StaticStruct()).template Get<T>();
 	}
 
+	UE_API FString ToDebugString() const;
+
 private:
 	TArray<FDataLinkOutputDataEntry> DataEntries;
 };
+
+#undef UE_API
diff --git a/Engine/Plugins/Experimental/DataLink/Source/DataLink/Public/DataLinkUtils.h b/Engine/Plugins/Experimental/DataLink/Source/DataLink/Public/DataLinkUtils.h
--- a/Engine/Plugins/Experimental/DataLink/Source/DataLink/Public/DataLinkUtils.h
+++ b/Engine/Plugins/Experimental/DataLink/Source/DataLink/Public/DataLinkUtils.h
@@ -73,6 +73,9 @@ namespace UE::DataLink
 	 * @param OutInputData the input data to update
 	 */
 	UE_API void SetInputData(UDataLinkGraph* InGraph, TArray<FDataLinkInputData>& OutInputData);
+
+	/** Converts the given struct data view to a debug string */
+	UE_API FString StructViewToDebugString(FConstStructView InDataView);
 }
 
 #undef UE_API
diff --git a/Engine/Plugins/Experimental/DataLink/Source/DataLink/Public/Nodes/DataLinkConstant.h b/Engine/Plugins/Experimental/DataLink/Source/DataLink/Public/Nodes/DataLinkConstant.h
--- a/Engine/Plugins/Experimental/DataLink/Source/DataLink/Public/Nodes/DataLinkConstant.h
+++ b/Engine/Plugins/Experimental/DataLink/Source/DataLink/Public/Nodes/DataLinkConstant.h
@@ -18,6 +18,7 @@ public:
 protected:
 	//~ Begin UDataLinkNode
 #if WITH_EDITOR
+	DATALINK_API virtual void OnFixupNode() override;
 	DATALINK_API virtual void OnBuildMetadata(FDataLinkNodeMetadata& Metadata) const override;
 #endif
 	DATALINK_API virtual void OnBuildPins(FDataLinkPinBuilder& Inputs, FDataLinkPinBuilder& Outputs) const override;
diff --git a/Engine/Plugins/Experimental/DataLink/Source/DataLinkEdGraph/Private/DataLinkEdGraph.cpp b/Engine/Plugins/Experimental/DataLink/Source/DataLinkEdGraph/Private/DataLinkEdGraph.cpp
--- a/Engine/Plugins/Experimental/DataLink/Source/DataLinkEdGraph/Private/DataLinkEdGraph.cpp
+++ b/Engine/Plugins/Experimental/DataLink/Source/DataLinkEdGraph/Private/DataLinkEdGraph.cpp
@@ -46,6 +46,11 @@ void UDataLinkEdGraph::BeginDestroy()
 	OnGraphCompiledHandle.Reset();
 }
 
+bool UDataLinkEdGraph::IsEditorOnly() const
+{
+	return true;
+}
+
 void UDataLinkEdGraph::DirtyGraph()
 {
 	// ChangeId could be regenerated every time graph is dirtied,
diff --git a/Engine/Plugins/Experimental/DataLink/Source/DataLinkEdGraph/Public/DataLinkEdGraph.h b/Engine/Plugins/Experimental/DataLink/Source/DataLinkEdGraph/Public/DataLinkEdGraph.h
--- a/Engine/Plugins/Experimental/DataLink/Source/DataLinkEdGraph/Public/DataLinkEdGraph.h
+++ b/Engine/Plugins/Experimental/DataLink/Source/DataLinkEdGraph/Public/DataLinkEdGraph.h
@@ -27,6 +27,7 @@ public:
 
 	//~ Begin UObject
 	DATALINKEDGRAPH_API virtual void BeginDestroy() override;
+	DATALINKEDGRAPH_API virtual bool IsEditorOnly() const override;
 	//~ End UObject
 
 	DATALINKEDGRAPH_API void DirtyGraph();
diff --git a/Engine/Plugins/Experimental/DataLink/Source/DataLinkEditor/Private/Compiler/DataLinkGraphCompiler.cpp b/Engine/Plugins/Experimental/DataLink/Source/DataLinkEditor/Private/Compiler/DataLinkGraphCompiler.cpp
--- a/Engine/Plugins/Experimental/DataLink/Source/DataLinkEditor/Private/Compiler/DataLinkGraphCompiler.cpp
+++ b/Engine/Plugins/Experimental/DataLink/Source/DataLinkEditor/Private/Compiler/DataLinkGraphCompiler.cpp
@@ -168,6 +168,7 @@ UDataLinkNode* FDataLinkGraphCompiler::CompileNode(UDataLinkNode* InTemplateNode
 		CompiledNode->OutputPins = MoveTemp(OutputPins);
 	}
 
+	CompiledNode->FixupNode();
 	return CompiledNode;
 }
 
diff --git a/Engine/Plugins/Experimental/DataLink/Source/DataLinkJson/Private/DataLinkStringToJson.cpp b/Engine/Plugins/Experimental/DataLink/Source/DataLinkJson/Private/DataLinkStringToJson.cpp
--- a/Engine/Plugins/Experimental/DataLink/Source/DataLinkJson/Private/DataLinkStringToJson.cpp
+++ b/Engine/Plugins/Experimental/DataLink/Source/DataLinkJson/Private/DataLinkStringToJson.cpp
@@ -3,6 +3,7 @@
 #include "DataLinkStringToJson.h"
 #include "DataLinkCoreTypes.h"
 #include "DataLinkExecutor.h"
+#include "DataLinkJsonLog.h"
 #include "DataLinkJsonNames.h"
 #include "DataLinkNames.h"
 #include "DataLinkNodeInstance.h"
@@ -45,6 +46,10 @@ EDataLinkExecutionReply UDataLinkStringToJson::OnExecute(FDataLinkExecutor& InEx
 
 	if (!OutputData.JsonObjectFromString(InputData.Value))
 	{
+		UE_LOG(LogDataLinkJson, Error, TEXT("[%.*s] Failed to create a JSON object from string %s")
+			, InExecutor.GetContextName().Len()
+			, InExecutor.GetContextName().GetData()
+			, *InputData.Value);
 		return EDataLinkExecutionReply::Unhandled;
 	}
 
diff --git a/Engine/Plugins/Experimental/SceneStateDataLink/Source/SceneStateDataLink/Private/SceneStateDataLinkLog.h b/Engine/Plugins/Experimental/SceneStateDataLink/Source/SceneStateDataLink/Private/SceneStateDataLinkLog.h
new file mode 100644
--- /dev/null
+++ b/Engine/Plugins/Experimental/SceneStateDataLink/Source/SceneStateDataLink/Private/SceneStateDataLinkLog.h
@@ -0,0 +1,7 @@
+// Copyright Epic Games, Inc. All Rights Reserved.
+
+#pragma once
+
+#include "Logging/LogMacros.h"
+
+DECLARE_LOG_CATEGORY_EXTERN(LogSceneStateDataLink, Log, All);
diff --git a/Engine/Plugins/Experimental/SceneStateDataLink/Source/SceneStateDataLink/Private/SceneStateDataLinkModule.cpp b/Engine/Plugins/Experimental/SceneStateDataLink/Source/SceneStateDataLink/Private/SceneStateDataLinkModule.cpp
--- a/Engine/Plugins/Experimental/SceneStateDataLink/Source/SceneStateDataLink/Private/SceneStateDataLinkModule.cpp
+++ b/Engine/Plugins/Experimental/SceneStateDataLink/Source/SceneStateDataLink/Private/SceneStateDataLinkModule.cpp
@@ -1,5 +1,8 @@
 // Copyright Epic Games, Inc. All Rights Reserved.
 
 #include "Modules/ModuleManager.h"
+#include "SceneStateDataLinkLog.h"
+
+DEFINE_LOG_CATEGORY(LogSceneStateDataLink);
 
 IMPLEMENT_MODULE(FDefaultModuleImpl, SceneStateDataLink)
diff --git a/Engine/Plugins/Experimental/SceneStateDataLink/Source/SceneStateDataLink/Private/SceneStateDataLinkRequestTask.cpp b/Engine/Plugins/Experimental/SceneStateDataLink/Source/SceneStateDataLink/Private/SceneStateDataLinkRequestTask.cpp
--- a/Engine/Plugins/Experimental/SceneStateDataLink/Source/SceneStateDataLink/Private/SceneStateDataLinkRequestTask.cpp
+++ b/Engine/Plugins/Experimental/SceneStateDataLink/Source/SceneStateDataLink/Private/SceneStateDataLinkRequestTask.cpp
@@ -6,6 +6,7 @@
 #include "DataLinkGraph.h"
 #include "DataLinkUtils.h"
 #include "Misc/ScopeExit.h"
+#include "SceneStateDataLinkLog.h"
 #include "SceneStateExecutionContext.h"
 #include "Tasks/SceneStateTaskExecutionContext.h"
 
@@ -72,6 +73,7 @@ void FSceneStateDataLinkRequestTask::OnExecutionFinished(const FDataLinkExecutor
 	UE::SceneState::FResolvePropertyResult Result;
 	if (!UE::SceneState::ResolveProperty(*ExecutionContext, Instance->OutputTarget, Result))
 	{
+		UE_LOG(LogSceneStateDataLink, Warning, TEXT("[%s] Data Link Request Task failed. Output Target could not be resolved."), *ExecutionContext->GetExecutionContextName());
 		return;
 	}
 
@@ -82,6 +84,10 @@ void FSceneStateDataLinkRequestTask::OnExecutionFinished(const FDataLinkExecutor
 	const FStructProperty* StructProperty = CastField<FStructProperty>(Result.ResolvedReference->SourceLeafProperty);
 	if (!StructProperty || !StructProperty->Struct || StructProperty->Struct != InOutputDataView.GetScriptStruct())
 	{
+		UE_LOG(LogSceneStateDataLink, Warning, TEXT("[%s] Data Link Request Task failed. Output Target struct ('%s') and Data Link Output struct ('%s') do not match!")
+			, *ExecutionContext->GetExecutionContextName()
+			, *GetNameSafe(StructProperty->Struct)
+			, *GetNameSafe(InOutputDataView.GetScriptStruct()));
 		return;
 	}
 
