Motion Design Data Link: separated OnExecutionFinished delegate into two: OnOutputData and OnExecutionFinished + reworded some terminology to not assume that a node has finished execution (e.g. "FinishNode" reworded to "Next")
OnOutputData gets broadcast when the data link graph has its output data ready without necessarily ending execution. This is because for new nodes that are push-approach, outputting data does not necessarily mean that the data link graph execution is finished.

diff --git a/Engine/Plugins/Experimental/AvalancheDataLink/Source/AvalancheDataLink/Private/AvaDataLinkInstance.cpp b/Engine/Plugins/Experimental/AvalancheDataLink/Source/AvalancheDataLink/Private/AvaDataLinkInstance.cpp
--- a/Engine/Plugins/Experimental/AvalancheDataLink/Source/AvalancheDataLink/Private/AvaDataLinkInstance.cpp
+++ b/Engine/Plugins/Experimental/AvalancheDataLink/Source/AvalancheDataLink/Private/AvaDataLinkInstance.cpp
@@ -87,22 +87,27 @@ namespace UE::AvaDataLink::Private
 
 void UAvaDataLinkInstance::Execute()
 {
-	if (Executor.IsValid())
-	{
-		UE_LOG(LogAvaDataLink, Error, TEXT("[%s] Data Link execution is in progress!"), Executor->GetContextName().GetData());
-		return;
-	}
+	Stop();
 
 	Executor = FDataLinkExecutor::Create(FDataLinkExecutorArguments(DataLinkInstance)
 #if WITH_DATALINK_CONTEXT
 		.SetContextName(BuildContextName())
 #endif
 		.SetContextObject(this)
-		.SetOnFinish(FOnDataLinkExecutionFinished::CreateUObject(this, &UAvaDataLinkInstance::OnExecutionFinished)));
+		.SetOnOutputData(FOnDataLinkOutputData::CreateUObject(this, &UAvaDataLinkInstance::OnOutputData)));
 
 	Executor->Run();
 }
 
+void UAvaDataLinkInstance::Stop()
+{
+	if (Executor.IsValid())
+	{
+		Executor->Stop();
+		Executor.Reset();
+	}
+}
+
 #if WITH_DATALINK_CONTEXT
 FString UAvaDataLinkInstance::BuildContextName() const
 {
@@ -110,11 +115,9 @@ FString UAvaDataLinkInstance::BuildContextName() const
 }
 #endif
 
-void UAvaDataLinkInstance::OnExecutionFinished(const FDataLinkExecutor& InExecutor, EDataLinkExecutionResult InResult, FConstStructView InOutputDataView)
+void UAvaDataLinkInstance::OnOutputData(const FDataLinkExecutor& InExecutor, FConstStructView InOutputDataView)
 {
-	Executor.Reset();
-
-	if (InResult == EDataLinkExecutionResult::Failed || !InOutputDataView.IsValid())
+	if (!InOutputDataView.IsValid())
 	{
 		return;
 	}
diff --git a/Engine/Plugins/Experimental/AvalancheDataLink/Source/AvalancheDataLink/Public/AvaDataLinkInstance.h b/Engine/Plugins/Experimental/AvalancheDataLink/Source/AvalancheDataLink/Public/AvaDataLinkInstance.h
--- a/Engine/Plugins/Experimental/AvalancheDataLink/Source/AvalancheDataLink/Public/AvaDataLinkInstance.h
+++ b/Engine/Plugins/Experimental/AvalancheDataLink/Source/AvalancheDataLink/Public/AvaDataLinkInstance.h
@@ -32,8 +32,12 @@ class UAvaDataLinkInstance : public UObject
 	GENERATED_BODY()
 
 public:
+ 	/** Starts a new active data link execution, stopping the existing if any */
 	void Execute();
 
+	/** Stops the active data link execution if any */
+	void Stop();
+
 	static FName GetDataLinkInstancePropertyName()
 	{
 		return GET_MEMBER_NAME_CHECKED(UAvaDataLinkInstance, DataLinkInstance);
@@ -49,7 +53,8 @@ private:
 	FString BuildContextName() const;
 #endif
 
-	void OnExecutionFinished(const FDataLinkExecutor& InExecutor, EDataLinkExecutionResult InResult, FConstStructView InOutputDataView);
+	/** Called when data link execution has received data */
+	void OnOutputData(const FDataLinkExecutor& InExecutor, FConstStructView InOutputDataView);
 
 	IAvaSceneInterface* GetSceneInterface() const;
 
diff --git a/Engine/Plugins/Experimental/DataLink/Source/DataLink/Private/DataLinkExecutor.cpp b/Engine/Plugins/Experimental/DataLink/Source/DataLink/Private/DataLinkExecutor.cpp
--- a/Engine/Plugins/Experimental/DataLink/Source/DataLink/Private/DataLinkExecutor.cpp
+++ b/Engine/Plugins/Experimental/DataLink/Source/DataLink/Private/DataLinkExecutor.cpp
@@ -29,7 +29,8 @@ TSharedPtr<FDataLinkExecutor> FDataLinkExecutor::Create(FDataLinkExecutorArgumen
 	Executor->ContextObject = MoveTemp(InArgs.ContextObject);
 	Executor->Instance = MoveTemp(InArgs.Instance);
 	Executor->Sink = MoveTemp(InArgs.Sink);
-	Executor->OnFinishDelegate = MoveTemp(InArgs.OnFinishDelegate);
+	Executor->OnOutputData = MoveTemp(InArgs.OnOutputData);
+	Executor->OnExecutionFinished = MoveTemp(InArgs.OnExecutionFinished);
 
 	return Executor;
 }
@@ -39,6 +40,15 @@ FDataLinkExecutor::FDataLinkExecutor(FPrivateToken)
 {
 }
 
+FDataLinkExecutor::~FDataLinkExecutor()
+{
+	if (bRunning && UObjectInitialized())
+	{
+		bRunning = false;
+		StopInternal();
+	}
+}
+
 const FGuid& FDataLinkExecutor::GetExecutorId() const
 {
 	return ExecutorId;
@@ -81,11 +91,15 @@ void FDataLinkExecutor::AddReferencedObjects(FReferenceCollector& InCollector)
 
 void FDataLinkExecutor::Run()
 {
+	if (bRunning)
+	{
+		return;
+	}
 	bRunning = true;
 
 	if (!ValidateRun())
 	{
-		Finish(EDataLinkExecutionResult::Failed, FConstStructView());
+		BroadcastFailure();
 		return;
 	}
 
@@ -100,8 +114,24 @@ void FDataLinkExecutor::Run()
 
 	if (!ExecuteEntryNodes())
 	{
-		Finish(EDataLinkExecutionResult::Failed, FConstStructView());
+		BroadcastFailure();
+	}
+}
+
+void FDataLinkExecutor::Stop(EDataLinkExecutionResult InExecutionResult)
+{
+	if (!bRunning)
+	{
+		return;
 	}
+	bRunning = false;
+
+	StopInternal();
+
+	// Hold onto an extra shared reference to allow implementers to reset their held shared reference in the finish delegate
+	// without destroying this
+	const TSharedRef<FDataLinkExecutor> This = SharedThis(this);
+	OnExecutionFinished.ExecuteIfBound(*this, InExecutionResult);
 }
 
 const FDataLinkNodeInstance& FDataLinkExecutor::GetNodeInstance(const UDataLinkNode* InNode) const
@@ -119,19 +149,24 @@ FDataLinkNodeInstance* FDataLinkExecutor::FindNodeInstanceMutable(const UDataLin
 	return NodeInstanceMap.Find(InNode);
 }
 
-void FDataLinkExecutor::SucceedNode(const UDataLinkNode* InNode, FConstStructView InOutputDataView)
+void FDataLinkExecutor::Next(const UDataLinkNode* InNode)
 {
-	if (!ProcessFinishedNode(InNode))
+	if (!ValidateNodeExecutionStatus(InNode))
 	{
-		Finish(EDataLinkExecutionResult::Failed, FConstStructView());
+		BroadcastFailure();
 		return;
 	}
 
 	TConstArrayView<FDataLinkPin> OutputPins = InNode->GetOutputPins();
 
-	TArray<FDataLinkPinReference> InputPinsToExecute;
+	TArray<FDataLinkPinReference, TInlineAllocator<1>> InputPinsToExecute;
 	InputPinsToExecute.Reserve(OutputPins.Num());
 
+	TArray<FConstStructView, TInlineAllocator<1>> InputDataViews;
+	InputDataViews.Reserve(OutputPins.Num());
+
+	const FDataLinkOutputDataViewer& OutputDataViewer = GetNodeInstance(InNode).GetOutputDataViewer();
+
 	for (const FDataLinkPin& OutputPin : OutputPins)
 	{
 		if (const FDataLinkPin* InputPin = OutputPin.GetLinkedInputPin())
@@ -139,39 +174,44 @@ void FDataLinkExecutor::SucceedNode(const UDataLinkNode* InNode, FConstStructVie
 			// If the Linked Input Pin was gotten successfully, then it means that the Output Pin's Linked Node was valid
 			check(OutputPin.LinkedNode);
 			InputPinsToExecute.Emplace(OutputPin.LinkedNode, InputPin);
+			InputDataViews.Emplace(OutputDataViewer.Find(OutputPin.Name));
 		}
 	}
 
-	// No next node, finish execution
+	// No next node, broadcast output data
 	if (InputPinsToExecute.IsEmpty())
 	{
 		// If the next node is null, it should be the output node
 		if (!ensureAlways(Instance.DataLinkGraph && Instance.DataLinkGraph->GetOutputNode() == InNode))
 		{
 			UE_LOG(LogDataLink, Error, TEXT("[%s] Node ('%s') has no node to go to and it's not an output node!"), GetContextName().GetData(), *InNode->GetName());
-			Finish(EDataLinkExecutionResult::Failed, FConstStructView());
+			BroadcastFailure();
 			return;
 		}
 
-		Finish(EDataLinkExecutionResult::Succeeded, InOutputDataView);
+		FStructView ResultDataView = OutputDataViewer.Find(Instance.DataLinkGraph->GetOutputPinName());
+		BroadcastOutputData(ResultDataView);
 		return;
 	}
 
-	TArray<FConstStructView> InputDataViews;
-	InputDataViews.Add(InOutputDataView);
-
-	if (!ExecuteInputPins(InputPinsToExecute, MoveTemp(InputDataViews)))
+	if (!ExecuteInputPins(InputPinsToExecute, InputDataViews))
 	{
-		Finish(EDataLinkExecutionResult::Failed, FConstStructView());
+		BroadcastFailure();
 	}
 }
 
-void FDataLinkExecutor::FailNode(const UDataLinkNode* InNode)
+void FDataLinkExecutor::NextPersist(const UDataLinkNode* InNode)
+{
+	bPersistExecution = true;
+	Next(InNode);
+}
+
+void FDataLinkExecutor::Fail(const UDataLinkNode* InNode)
 {
-	ProcessFinishedNode(InNode);
+	ValidateNodeExecutionStatus(InNode);
 
 	UE_LOG(LogDataLink, Log, TEXT("[%s] Node ('%s') has failed to execute. Data Link execution finished."), GetContextName().GetData(), *InNode->GetName());
-	Finish(EDataLinkExecutionResult::Failed, FConstStructView());
+	BroadcastFailure();
 }
 
 bool FDataLinkExecutor::ValidateRun()
@@ -280,7 +320,7 @@ bool FDataLinkExecutor::ExecuteEntryNodes()
 	return NodesExecuted > 0;
 }
 
-bool FDataLinkExecutor::ExecuteInputPins(TConstArrayView<FDataLinkPinReference> InInputPins, TArray<FConstStructView>&& InInputDataViews, uint16* OutNodesExecuted)
+bool FDataLinkExecutor::ExecuteInputPins(TConstArrayView<FDataLinkPinReference> InInputPins, TConstArrayView<FConstStructView> InInputDataViews, uint16* OutNodesExecuted)
 {
 	if (!ValidateInputPins(InInputPins, InInputDataViews))
 	{
@@ -375,6 +415,21 @@ bool FDataLinkExecutor::ExecuteNode(const UDataLinkNode& InNode)
 	return true;
 }
 
+void FDataLinkExecutor::StopInternal()
+{
+	// Stop active node instances
+	for (const TPair<TObjectPtr<const UDataLinkNode>, FDataLinkNodeInstance>& Pair : NodeInstanceMap)
+	{
+		if (Pair.Key && Pair.Value.Status == EDataLinkNodeStatus::Executing)
+		{
+			Pair.Key->Stop(*this);
+		}
+	}
+
+	NodeInstanceMap.Reset();
+	Sink.Reset();
+}
+
 FDataLinkNodeInstance& FDataLinkExecutor::FindOrAddNodeInstance(const UDataLinkNode* InNode)
 {
 	check(InNode);
@@ -386,16 +441,16 @@ FDataLinkNodeInstance& FDataLinkExecutor::FindOrAddNodeInstance(const UDataLinkN
 	return *NodeInstance;
 }
 
-bool FDataLinkExecutor::ProcessFinishedNode(const UDataLinkNode* InNode)
+bool FDataLinkExecutor::ValidateNodeExecutionStatus(const UDataLinkNode* InNode)
 {
 	FDataLinkNodeInstance* const NodeInstance = FindNodeInstanceMutable(InNode);
-	if (!ensureAlwaysMsgf(NodeInstance, TEXT("[%.*s] FinishNode called on a node that no longer has node instance data!"), GetContextName().Len(), GetContextName().GetData()))
+	if (!ensureAlwaysMsgf(NodeInstance, TEXT("[%.*s] Invalid Node Status. No node instance data found!"), GetContextName().Len(), GetContextName().GetData()))
 	{
 		return false;
 	}
 
 	const bool bWasExecuting = NodeInstance->Status == EDataLinkNodeStatus::Executing;
-	if (!ensureAlwaysMsgf(bWasExecuting, TEXT("[%.*s] FinishNode called on a node that was not executing!"), GetContextName().Len(), GetContextName().GetData()))
+	if (!ensureAlwaysMsgf(bWasExecuting, TEXT("[%.*s] Invalid Node Status. Node was not executing!"), GetContextName().Len(), GetContextName().GetData()))
 	{
 		return false;
 	}
@@ -415,26 +470,30 @@ bool FDataLinkExecutor::ProcessFinishedNode(const UDataLinkNode* InNode)
 			, *NodeInstance->GetOutputDataViewer().ToDebugString());
 	}
 #endif
-
-	NodeInstance->Status = EDataLinkNodeStatus::Finished;
 	return true;
 }
 
-void FDataLinkExecutor::Finish(EDataLinkExecutionResult InStatus, FConstStructView InOutputData)
+void FDataLinkExecutor::BroadcastFailure()
+{
+	BroadcastOutputData(FConstStructView());
+	Stop(EDataLinkExecutionResult::Failed);
+}
+
+void FDataLinkExecutor::BroadcastOutputData(FConstStructView InOutputData)
 {
-	// Finish has already been called
+	// Ensure executor was running
 	if (!bRunning)
 	{
 		return;
 	}
 
-	// Hold onto an extra shared reference to allow implementers to reset their held shared reference in the finish delegate
-	// without destroying this
-	TSharedRef<FDataLinkExecutor> This = SharedThis(this);
+	OnOutputData.ExecuteIfBound(*this, InOutputData);
 
-	OnFinishDelegate.ExecuteIfBound(*this, InStatus, InOutputData);
-	NodeInstanceMap.Reset();
-	Sink.Reset();
+	if (!bPersistExecution)
+	{
+		Stop(EDataLinkExecutionResult::Succeeded);
+	}
 
-	bRunning = false;
+	// Reset on graph end
+	bPersistExecution = false;
 }
diff --git a/Engine/Plugins/Experimental/DataLink/Source/DataLink/Private/DataLinkExecutorArguments.cpp b/Engine/Plugins/Experimental/DataLink/Source/DataLink/Private/DataLinkExecutorArguments.cpp
--- a/Engine/Plugins/Experimental/DataLink/Source/DataLink/Private/DataLinkExecutorArguments.cpp
+++ b/Engine/Plugins/Experimental/DataLink/Source/DataLink/Private/DataLinkExecutorArguments.cpp
@@ -33,8 +33,14 @@ FDataLinkExecutorArguments&& FDataLinkExecutorArguments::SetSink(const TSharedPt
 	return MoveTemp(*this);
 }
 
-FDataLinkExecutorArguments&& FDataLinkExecutorArguments::SetOnFinish(FOnDataLinkExecutionFinished&& InDelegate)
+FDataLinkExecutorArguments&& FDataLinkExecutorArguments::SetOnOutputData(FOnDataLinkOutputData&& InDelegate)
 {
-	OnFinishDelegate = MoveTemp(InDelegate);
+	OnOutputData = MoveTemp(InDelegate);
+	return MoveTemp(*this);
+}
+
+FDataLinkExecutorArguments&& FDataLinkExecutorArguments::SetOnFinished(FOnDataLinkExecutionFinished&& InDelegate)
+{
+	OnExecutionFinished = MoveTemp(InDelegate);
 	return MoveTemp(*this);
 }
diff --git a/Engine/Plugins/Experimental/DataLink/Source/DataLink/Private/DataLinkNode.cpp b/Engine/Plugins/Experimental/DataLink/Source/DataLink/Private/DataLinkNode.cpp
--- a/Engine/Plugins/Experimental/DataLink/Source/DataLink/Private/DataLinkNode.cpp
+++ b/Engine/Plugins/Experimental/DataLink/Source/DataLink/Private/DataLinkNode.cpp
@@ -14,10 +14,15 @@ void UDataLinkNode::Execute(FDataLinkExecutor& InExecutor) const
 	if (Reply == EDataLinkExecutionReply::Unhandled)
 	{
 		UE_LOG(LogDataLink, Log, TEXT("[%s] Node '%s' did not handle execution."), InExecutor.GetContextName().GetData(), *GetName())
-		InExecutor.FailNode(this);
+		InExecutor.Fail(this);
 	}
 }
 
+void UDataLinkNode::Stop(const FDataLinkExecutor& InExecutor) const
+{
+	OnStop(InExecutor);
+}
+
 #if WITH_EDITOR
 void UDataLinkNode::FixupNode()
 {
diff --git a/Engine/Plugins/Experimental/DataLink/Source/DataLink/Private/DataLinkRequestProxy.cpp b/Engine/Plugins/Experimental/DataLink/Source/DataLink/Private/DataLinkRequestProxy.cpp
--- a/Engine/Plugins/Experimental/DataLink/Source/DataLink/Private/DataLinkRequestProxy.cpp
+++ b/Engine/Plugins/Experimental/DataLink/Source/DataLink/Private/DataLinkRequestProxy.cpp
@@ -2,6 +2,7 @@
 
 #include "DataLinkRequestProxy.h"
 #include "Containers/Ticker.h"
+#include "DataLinkEnums.h"
 #include "DataLinkExecutor.h"
 #include "DataLinkExecutorArguments.h"
 #include "DataLinkUtils.h"
@@ -22,7 +23,7 @@ void UDataLinkRequestProxy::ProcessRequest(FDataLinkInstance&& InDataLinkInstanc
 #endif
 		.SetContextObject(InExecutionContext)
 		.SetSink(UE::DataLink::TryGetSink(InDataLinkSinkProvider))
-		.SetOnFinish(FOnDataLinkExecutionFinished::CreateUObject(this, &UDataLinkRequestProxy::OnExecutionFinished)));
+		.SetOnOutputData(FOnDataLinkOutputData::CreateUObject(this, &UDataLinkRequestProxy::OnOutputDataReceived)));
 
 	// Delay execution by 1 frame. This is so that even for executions that finish immediately,
 	// the Async Call blueprint script still has chance to get the returned Proxy and listen to execution finished event
@@ -39,13 +40,15 @@ void UDataLinkRequestProxy::ProcessRequest(FDataLinkInstance&& InDataLinkInstanc
 		}));
 }
 
-void UDataLinkRequestProxy::OnExecutionFinished(const FDataLinkExecutor& InExecutor, EDataLinkExecutionResult InResult, FConstStructView InOutputData)
+void UDataLinkRequestProxy::OnOutputDataReceived(const FDataLinkExecutor& InExecutor, FConstStructView InOutputData)
 {
-	DataLinkExecutor.Reset();
-
 	FInstancedStruct OutputData(InOutputData);
 
+	const EDataLinkExecutionResult ExecutionResult = InOutputData.IsValid()
+		? EDataLinkExecutionResult::Succeeded
+		: EDataLinkExecutionResult::Failed;
+
 	// Allow Blueprint CallInEditor functions to handle then execution result
 	FEditorScriptExecutionGuard ScriptGuard;
-	OnRequestComplete.Broadcast(OutputData, InResult);
+	OnOutputData.Broadcast(OutputData, ExecutionResult);
 }
diff --git a/Engine/Plugins/Experimental/DataLink/Source/DataLink/Private/DataLinkUtils.cpp b/Engine/Plugins/Experimental/DataLink/Source/DataLink/Private/DataLinkUtils.cpp
--- a/Engine/Plugins/Experimental/DataLink/Source/DataLink/Private/DataLinkUtils.cpp
+++ b/Engine/Plugins/Experimental/DataLink/Source/DataLink/Private/DataLinkUtils.cpp
@@ -18,6 +18,22 @@
 namespace UE::DataLink
 {
 
+bool CopyDataView(FStructView InDestView, FConstStructView InSourceView)
+{
+	if (!InDestView.IsValid() || !InSourceView.IsValid())
+	{
+		return false;
+	}
+
+	if (InDestView.GetScriptStruct() != InSourceView.GetScriptStruct())
+	{
+		return false;
+	}
+
+	InDestView.GetScriptStruct()->CopyScriptStruct(InDestView.GetMemory(), InSourceView.GetMemory());
+	return true;
+}
+
 bool ReplaceObject(UObject*& InOutObject, UObject* InOuter, UClass* InClass)
 {
 	// Check if the object class already matches the new class
diff --git a/Engine/Plugins/Experimental/DataLink/Source/DataLink/Private/Nodes/DataLinkConstant.cpp b/Engine/Plugins/Experimental/DataLink/Source/DataLink/Private/Nodes/DataLinkConstant.cpp
--- a/Engine/Plugins/Experimental/DataLink/Source/DataLink/Private/Nodes/DataLinkConstant.cpp
+++ b/Engine/Plugins/Experimental/DataLink/Source/DataLink/Private/Nodes/DataLinkConstant.cpp
@@ -7,6 +7,7 @@
 #include "DataLinkNodeInstance.h"
 #include "DataLinkNodeMetadata.h"
 #include "DataLinkPinBuilder.h"
+#include "DataLinkUtils.h"
 
 void UDataLinkConstant::SetStruct(const UScriptStruct* InStructType)
 {
@@ -104,7 +105,7 @@ EDataLinkExecutionReply UDataLinkConstant::OnExecute(FDataLinkExecutor& InExecut
 	const FDataLinkOutputDataViewer& OutputDataViewer = NodeInstance.GetOutputDataViewer();
 
 	const FStructView OutputDataView = OutputDataViewer.Find(UE::DataLink::OutputDefault, Instance.GetScriptStruct());
-	if (Instance.GetScriptStruct() != OutputDataView.GetScriptStruct())
+	if (!UE::DataLink::CopyDataView(OutputDataView, Instance))
 	{
 		UE_LOG(LogDataLink, Error, TEXT("[%.*s] DataLinkConstant failed. %s does not match output struct %s")
 			, InExecutor.GetContextName().Len(), InExecutor.GetContextName().GetData()
@@ -113,7 +114,6 @@ EDataLinkExecutionReply UDataLinkConstant::OnExecute(FDataLinkExecutor& InExecut
 		return EDataLinkExecutionReply::Unhandled;
 	}
 
-	OutputDataView.GetScriptStruct()->CopyScriptStruct(OutputDataView.GetMemory(), Instance.GetMemory());
-	InExecutor.SucceedNode(this, OutputDataView);
+	InExecutor.Next(this);
 	return EDataLinkExecutionReply::Handled;
 }
diff --git a/Engine/Plugins/Experimental/DataLink/Source/DataLink/Private/Nodes/Script/DataLinkScriptNode.cpp b/Engine/Plugins/Experimental/DataLink/Source/DataLink/Private/Nodes/Script/DataLinkScriptNode.cpp
--- a/Engine/Plugins/Experimental/DataLink/Source/DataLink/Private/Nodes/Script/DataLinkScriptNode.cpp
+++ b/Engine/Plugins/Experimental/DataLink/Source/DataLink/Private/Nodes/Script/DataLinkScriptNode.cpp
@@ -25,7 +25,7 @@ bool UDataLinkScriptNode::Succeed(const FInstancedStruct& OutputData)
 	const UScriptStruct* OutputDataStruct = OutputData.GetScriptStruct();
 	if (!OutputDataStruct || OutputPin.Struct != OutputDataStruct)
 	{
-		Executor->FailNode(Node);
+		Executor->Fail(Node);
 		return false;
 	}
 
@@ -36,12 +36,12 @@ bool UDataLinkScriptNode::Succeed(const FInstancedStruct& OutputData)
 	const FStructView OutputDataView = OutputDataViewer.Find(OutputPin.Name, OutputPin.Struct);
 	if (!OutputDataView.IsValid())
 	{
-		Executor->FailNode(Node);
+		Executor->Fail(Node);
 		return false;
 	}
 
 	OutputDataStruct->CopyScriptStruct(OutputDataView.GetMemory(), OutputData.GetMemory());
-	Executor->SucceedNode(Node, OutputDataView);
+	Executor->Next(Node);
 	return true;
 }
 
@@ -49,7 +49,7 @@ void UDataLinkScriptNode::Fail()
 {
 	if (TSharedPtr<FDataLinkExecutor> Executor = ExecutorWeak.Pin())
 	{
-		Executor->FailNode(Node);
+		Executor->Fail(Node);
 	}
 }
 
@@ -127,12 +127,12 @@ DEFINE_FUNCTION(UDataLinkScriptNode::execSucceedWildcard)
 		if (OutputDataView.IsValid())
 		{
 			OutputStructProperty->Struct->CopyScriptStruct(OutputDataView.GetMemory(), OutputStructAddress);
-			Executor->SucceedNode(P_THIS->Node, OutputDataView);
+			Executor->Next(P_THIS->Node);
 			*(bool*)RESULT_PARAM = true;
 		}
 		else
 		{
-			Executor->FailNode(P_THIS->Node);
+			Executor->Fail(P_THIS->Node);
 		}
 	}
 
diff --git a/Engine/Plugins/Experimental/DataLink/Source/DataLink/Private/Nodes/String/DataLinkNodeStringBuilder.cpp b/Engine/Plugins/Experimental/DataLink/Source/DataLink/Private/Nodes/String/DataLinkNodeStringBuilder.cpp
--- a/Engine/Plugins/Experimental/DataLink/Source/DataLink/Private/Nodes/String/DataLinkNodeStringBuilder.cpp
+++ b/Engine/Plugins/Experimental/DataLink/Source/DataLink/Private/Nodes/String/DataLinkNodeStringBuilder.cpp
@@ -27,11 +27,11 @@ EDataLinkExecutionReply UDataLinkNodeStringBuilder::OnExecute(FDataLinkExecutor&
 	const FDataLinkStringBuilder StringBuilder(Segments, Tokens);
 	if (StringBuilder.BuildString(NodeInstance.GetInputDataViewer(), OutputString.Value))
 	{
-		InExecutor.SucceedNode(this, FConstStructView::Make(OutputString));
+		InExecutor.Next(this);
 	}
 	else
 	{
-		InExecutor.FailNode(this);
+		InExecutor.Fail(this);
 	}
 
 	return EDataLinkExecutionReply::Handled;
diff --git a/Engine/Plugins/Experimental/DataLink/Source/DataLink/Public/DataLinkDelegates.h b/Engine/Plugins/Experimental/DataLink/Source/DataLink/Public/DataLinkDelegates.h
--- a/Engine/Plugins/Experimental/DataLink/Source/DataLink/Public/DataLinkDelegates.h
+++ b/Engine/Plugins/Experimental/DataLink/Source/DataLink/Public/DataLinkDelegates.h
@@ -9,4 +9,5 @@ class FDataLinkExecutor;
 enum class EDataLinkExecutionResult : uint8;
 struct FConstStructView;
 
-DECLARE_DELEGATE_ThreeParams(FOnDataLinkExecutionFinished, const FDataLinkExecutor&, EDataLinkExecutionResult, FConstStructView);
+DECLARE_DELEGATE_TwoParams(FOnDataLinkOutputData, const FDataLinkExecutor&, FConstStructView);
+DECLARE_DELEGATE_TwoParams(FOnDataLinkExecutionFinished, const FDataLinkExecutor&, EDataLinkExecutionResult);
diff --git a/Engine/Plugins/Experimental/DataLink/Source/DataLink/Public/DataLinkExecutor.h b/Engine/Plugins/Experimental/DataLink/Source/DataLink/Public/DataLinkExecutor.h
--- a/Engine/Plugins/Experimental/DataLink/Source/DataLink/Public/DataLinkExecutor.h
+++ b/Engine/Plugins/Experimental/DataLink/Source/DataLink/Public/DataLinkExecutor.h
@@ -3,6 +3,7 @@
 #pragma once
 
 #include "DataLinkDelegates.h"
+#include "DataLinkEnums.h"
 #include "DataLinkInstance.h"
 #include "DataLinkNodeInstance.h"
 #include "DataLinkSink.h"
@@ -28,6 +29,8 @@ public:
 
 	explicit FDataLinkExecutor(FPrivateToken);
 
+	virtual ~FDataLinkExecutor() override;
+
 	/** Gets the unique identifier for this Executor */
 	DATALINK_API const FGuid& GetExecutorId() const;
 
@@ -41,9 +44,17 @@ public:
 	/** Gets the object that is responsible for this execution */
 	DATALINK_API UObject* GetContextObject() const;
 
+	bool IsRunning() const
+	{
+		return bRunning;
+	}
+
 	/** Called to start the graph execution */
 	DATALINK_API void Run();
 
+	/** Called to stop graph execution */
+	DATALINK_API void Stop(EDataLinkExecutionResult InExecutionResult = EDataLinkExecutionResult::Succeeded);
+
 	/** Retrieves the node instance data of this execution for the given node */
 	DATALINK_API const FDataLinkNodeInstance& GetNodeInstance(const UDataLinkNode* InNode) const;
 
@@ -54,19 +65,26 @@ public:
 	DATALINK_API FDataLinkNodeInstance* FindNodeInstanceMutable(const UDataLinkNode* InNode);
 
 	/**
-	 * Called when a node has finished executing successfully
-	 * If it's the output node that has finished executing, it will finish the entire graph execution
-	 * @param InNode the node that finished execution
-	 * @param InOutputDataView the data view of the output (if succeeded)
+	 * Called when a node has received data successfully and is ready to pass it to the next node
+	 * If this is called from the output node, it will broadcast this data to the listeners instead
+	 * @param InNode the node that received this data
+	 */
+	DATALINK_API void Next(const UDataLinkNode* InNode);
+
+	/**
+	 * Used by Nodes that keep receiving data (instead of being a one-off) to keep execution active.
+	 * Called when a node has received data successfully and is ready to pass it to the next node
+	 * If this is called from the output node, it will broadcast this data to the listeners instead.
+	 * @param InNode the node that received this data
 	 */
-	DATALINK_API void SucceedNode(const UDataLinkNode* InNode, FConstStructView InOutputDataView);
+	DATALINK_API void NextPersist(const UDataLinkNode* InNode);
 
 	/**
-	 * Called when a node has failed to execute and produce a valid output
+	 * Called when a node has failed to receive data and produce a valid output
 	 * This stops the entire graph execution
 	 * @param InNode the node that failed
 	 */
-	DATALINK_API void FailNode(const UDataLinkNode* InNode);
+	DATALINK_API void Fail(const UDataLinkNode* InNode);
 
 private:
 	/** Called when running to ensure all the parameters are properly set. Returns true if run can be done, false otherwise */
@@ -85,17 +103,26 @@ private:
 	 * @param OutNodesExecuted the number of nodes that got executed
 	 * @return true if the input pins were processed without error
 	 */
-	bool ExecuteInputPins(TConstArrayView<FDataLinkPinReference> InInputPins, TArray<FConstStructView>&& InInputDataViews, uint16* OutNodesExecuted = nullptr);
+	bool ExecuteInputPins(TConstArrayView<FDataLinkPinReference> InInputPins, TConstArrayView<FConstStructView> InInputDataViews, uint16* OutNodesExecuted = nullptr);
 
 	bool ExecuteNode(const UDataLinkNode& InNode);
 
+	/** Called to stop graph execution */
+	void StopInternal();
+
 	FDataLinkNodeInstance& FindOrAddNodeInstance(const UDataLinkNode* InNode);
 
-	/** Called when a node has finished */
-	bool ProcessFinishedNode(const UDataLinkNode* InNode);
+	/**
+	 * Called when a node has executed and finished its logic.
+	 * Checks that the node instance is valid and executing
+	 */
+	bool ValidateNodeExecutionStatus(const UDataLinkNode* InNode);
+
+	/** Called when there was a failure in the graph execution */
+	void BroadcastFailure();
 
 	/** Called when the entire graph execution has completed */
-	void Finish(EDataLinkExecutionResult InStatus, FConstStructView InOutputData);
+	void BroadcastOutputData(FConstStructView InOutputData);
 
 	//~ Begin FGCObject
 	virtual FString GetReferencerName() const override;
@@ -119,11 +146,17 @@ private:
 	/** Sink where all the data for this execution is queried and stored */
 	TSharedPtr<FDataLinkSink> Sink;
 
-	/** Delegate to call when the execution completes */
-	FOnDataLinkExecutionFinished OnFinishDelegate;
+	/** Delegate to call when the data link outputs data */
+	FOnDataLinkOutputData OnOutputData;
+
+	/** Delegate to call when the data link execution finishes */
+	FOnDataLinkExecutionFinished OnExecutionFinished;
 
 	/** Map of a Node to its Instance (Input, Output and Instance Data) for this execution */
 	mutable TMap<TObjectPtr<const UDataLinkNode>, FDataLinkNodeInstance> NodeInstanceMap;
 
+	/** Whether to keep execution running upon data broadcast, or finish */
+	bool bPersistExecution = false;
+
 	bool bRunning = false;
 };
diff --git a/Engine/Plugins/Experimental/DataLink/Source/DataLink/Public/DataLinkExecutorArguments.h b/Engine/Plugins/Experimental/DataLink/Source/DataLink/Public/DataLinkExecutorArguments.h
--- a/Engine/Plugins/Experimental/DataLink/Source/DataLink/Public/DataLinkExecutorArguments.h
+++ b/Engine/Plugins/Experimental/DataLink/Source/DataLink/Public/DataLinkExecutorArguments.h
@@ -30,8 +30,11 @@ public:
 	/** Sets the Sink to use. If unset, the Executor will create its own temporary sink */
 	DATALINK_API FDataLinkExecutorArguments&& SetSink(const TSharedPtr<FDataLinkSink>& InSink);
 
+	/** The delegate to call when the execution outputs data */
+	DATALINK_API FDataLinkExecutorArguments&& SetOnOutputData(FOnDataLinkOutputData&& InDelegate);
+
 	/** The delegate to call when the execution finishes */
-	DATALINK_API FDataLinkExecutorArguments&& SetOnFinish(FOnDataLinkExecutionFinished&& InDelegate);
+	DATALINK_API FDataLinkExecutorArguments&& SetOnFinished(FOnDataLinkExecutionFinished&& InDelegate);
 
 private:
 	FDataLinkInstance Instance;
@@ -44,5 +47,9 @@ private:
 
 	TSharedPtr<FDataLinkSink> Sink;
 
-	FOnDataLinkExecutionFinished OnFinishDelegate;
+	/** Delegate called when data link execution has output data ready */
+	FOnDataLinkOutputData OnOutputData;
+
+	/** Delegate called when data link execution has finished completely */
+	FOnDataLinkExecutionFinished OnExecutionFinished;
 };
diff --git a/Engine/Plugins/Experimental/DataLink/Source/DataLink/Public/DataLinkGraph.h b/Engine/Plugins/Experimental/DataLink/Source/DataLink/Public/DataLinkGraph.h
--- a/Engine/Plugins/Experimental/DataLink/Source/DataLink/Public/DataLinkGraph.h
+++ b/Engine/Plugins/Experimental/DataLink/Source/DataLink/Public/DataLinkGraph.h
@@ -42,6 +42,11 @@ public:
 		return OutputNode;
 	}
 
+	FName GetOutputPinName() const
+	{
+		return OutputPinName;
+	}
+
 #if WITH_EDITOR
 	UEdGraph* GetEdGraph() const
 	{
@@ -71,6 +76,10 @@ private:
 	UPROPERTY()
 	TObjectPtr<UDataLinkNode> OutputNode;
 
+	/** Output Node Pin name that provides the result data */
+	UPROPERTY()
+	FName OutputPinName;
+
 #if WITH_EDITORONLY_DATA
 	/** EdGraph used to compile the nodes in this graph */
 	UPROPERTY()
diff --git a/Engine/Plugins/Experimental/DataLink/Source/DataLink/Public/DataLinkNode.h b/Engine/Plugins/Experimental/DataLink/Source/DataLink/Public/DataLinkNode.h
--- a/Engine/Plugins/Experimental/DataLink/Source/DataLink/Public/DataLinkNode.h
+++ b/Engine/Plugins/Experimental/DataLink/Source/DataLink/Public/DataLinkNode.h
@@ -32,6 +32,9 @@ class UDataLinkNode : public UObject
 public:
 	void Execute(FDataLinkExecutor& InExecutor) const;
 
+	/** Called when the node is still executing and is forced to stop execution and clean up its resources */
+	void Stop(const FDataLinkExecutor& InExecutor) const;
+
 	TConstArrayView<FDataLinkPin> GetInputPins() const
 	{
 		return InputPins;
@@ -81,6 +84,11 @@ protected:
 		return EDataLinkExecutionReply::Unhandled;
 	}
 
+	/** Called when the node is still executing and is forced to stop execution and clean up its resources */
+	virtual void OnStop(const FDataLinkExecutor& InExecutor) const
+	{
+	}
+
 	/** Optional instance data struct for data outside input and output */
 	UPROPERTY()
 	TObjectPtr<const UScriptStruct> InstanceStruct;
diff --git a/Engine/Plugins/Experimental/DataLink/Source/DataLink/Public/DataLinkNodeInstance.h b/Engine/Plugins/Experimental/DataLink/Source/DataLink/Public/DataLinkNodeInstance.h
--- a/Engine/Plugins/Experimental/DataLink/Source/DataLink/Public/DataLinkNodeInstance.h
+++ b/Engine/Plugins/Experimental/DataLink/Source/DataLink/Public/DataLinkNodeInstance.h
@@ -14,7 +14,6 @@ enum class EDataLinkNodeStatus : uint8
 {
 	NotStarted,
 	Executing,
-	Finished,
 };
 
 struct FDataLinkNodeInstance
diff --git a/Engine/Plugins/Experimental/DataLink/Source/DataLink/Public/DataLinkOutputDataViewer.h b/Engine/Plugins/Experimental/DataLink/Source/DataLink/Public/DataLinkOutputDataViewer.h
--- a/Engine/Plugins/Experimental/DataLink/Source/DataLink/Public/DataLinkOutputDataViewer.h
+++ b/Engine/Plugins/Experimental/DataLink/Source/DataLink/Public/DataLinkOutputDataViewer.h
@@ -33,7 +33,7 @@ class FDataLinkOutputDataViewer
 public:
 	explicit FDataLinkOutputDataViewer(TConstArrayView<FDataLinkPin> InOutputPins);
 
-	UE_API FStructView Find(FName InOutputName, const UScriptStruct* InDesiredStruct) const;
+	UE_API FStructView Find(FName InOutputName, const UScriptStruct* InDesiredStruct = nullptr) const;
 
 	int32 Num() const
 	{
diff --git a/Engine/Plugins/Experimental/DataLink/Source/DataLink/Public/DataLinkRequestProxy.h b/Engine/Plugins/Experimental/DataLink/Source/DataLink/Public/DataLinkRequestProxy.h
--- a/Engine/Plugins/Experimental/DataLink/Source/DataLink/Public/DataLinkRequestProxy.h
+++ b/Engine/Plugins/Experimental/DataLink/Source/DataLink/Public/DataLinkRequestProxy.h
@@ -11,7 +11,7 @@ class FDataLinkExecutor;
 enum class EDataLinkExecutionResult : uint8;
 struct FDataLinkInstance;
 
-DECLARE_DYNAMIC_MULTICAST_DELEGATE_TwoParams(FOnDataLinkRequestComplete, const FInstancedStruct&, OutputData, EDataLinkExecutionResult, ExecutionResult);
+DECLARE_DYNAMIC_MULTICAST_DELEGATE_TwoParams(FOnDataLinkOutputDataProxy, const FInstancedStruct&, OutputData, EDataLinkExecutionResult, ExecutionResult);
 
 UCLASS(MinimalAPI)
 class UDataLinkRequestProxy : public UObject
@@ -20,7 +20,7 @@ class UDataLinkRequestProxy : public UObject
 
 public:
 	UPROPERTY(BlueprintAssignable)
-	FOnDataLinkRequestComplete OnRequestComplete;
+	FOnDataLinkOutputDataProxy OnOutputData;
 
 	UFUNCTION(BlueprintCallable, Category="Data Link", meta=(BlueprintInternalUseOnly))
 	static DATALINK_API UDataLinkRequestProxy* CreateRequestProxy(FDataLinkInstance InDataLinkInstance
@@ -32,7 +32,7 @@ public:
 		, TScriptInterface<IDataLinkSinkProvider> InDataLinkSinkProvider);
 
 private:
-	void OnExecutionFinished(const FDataLinkExecutor& InExecutor, EDataLinkExecutionResult InResult, FConstStructView InOutputData);
+	void OnOutputDataReceived(const FDataLinkExecutor& InExecutor, FConstStructView InOutputData);
 
 	TSharedPtr<FDataLinkExecutor> DataLinkExecutor;
 };
diff --git a/Engine/Plugins/Experimental/DataLink/Source/DataLink/Public/DataLinkUtils.h b/Engine/Plugins/Experimental/DataLink/Source/DataLink/Public/DataLinkUtils.h
--- a/Engine/Plugins/Experimental/DataLink/Source/DataLink/Public/DataLinkUtils.h
+++ b/Engine/Plugins/Experimental/DataLink/Source/DataLink/Public/DataLinkUtils.h
@@ -23,6 +23,14 @@ template <typename T> class TScriptInterface;
 
 namespace UE::DataLink
 {
+	/**
+	 * Copies a given source view into dest view if the script structs match
+	 * @param InDestView the struct view to copy to
+	 * @param InSourceView the struct view to copy from
+	 * @return true if the struct views were compatible and the copy was performed
+	 */
+	UE_API bool CopyDataView(FStructView InDestView, FConstStructView InSourceView);
+
 	/**
 	 * Replaces a given Object with a new object with the same name but different class
 	 * @param InOutObject object to replace. Can come in as null
diff --git a/Engine/Plugins/Experimental/DataLink/Source/DataLinkDataTable/Private/DataLinkDataTableSource.cpp b/Engine/Plugins/Experimental/DataLink/Source/DataLinkDataTable/Private/DataLinkDataTableSource.cpp
--- a/Engine/Plugins/Experimental/DataLink/Source/DataLinkDataTable/Private/DataLinkDataTableSource.cpp
+++ b/Engine/Plugins/Experimental/DataLink/Source/DataLinkDataTable/Private/DataLinkDataTableSource.cpp
@@ -5,6 +5,7 @@
 #include "DataLinkNames.h"
 #include "DataLinkNodeInstance.h"
 #include "DataLinkPinBuilder.h"
+#include "DataLinkUtils.h"
 #include "StructUtils/StructView.h"
 
 void UDataLinkDataTableSource::OnBuildPins(FDataLinkPinBuilder& Inputs, FDataLinkPinBuilder& Outputs) const
@@ -23,6 +24,7 @@ EDataLinkExecutionReply UDataLinkDataTableSource::OnExecute(FDataLinkExecutor& I
 	const FDataLinkNodeInstance& NodeInstance = InExecutor.GetNodeInstance(this);
 
 	const FDataLinkInputDataViewer& InputDataViewer = NodeInstance.GetInputDataViewer();
+	const FDataLinkOutputDataViewer& OutputDataViewer = NodeInstance.GetOutputDataViewer();
 
 	const FDataTableRowHandle& RowHandle = InputDataViewer.Get<FDataTableRowHandle>(UE::DataLinkDataTable::InputRow);
 	if (!RowHandle.DataTable)
@@ -38,6 +40,13 @@ EDataLinkExecutionReply UDataLinkDataTableSource::OnExecute(FDataLinkExecutor& I
 
 	// Row memory will return null if the Row struct is null. 
 	check(RowHandle.DataTable->RowStruct != nullptr);
-	InExecutor.SucceedNode(this, FConstStructView(RowHandle.DataTable->RowStruct, RowMemory));
+
+	const FStructView OutputDataView = OutputDataViewer.Find(UE::DataLink::OutputDefault, RowHandle.DataTable->RowStruct);
+	if (!UE::DataLink::CopyDataView(OutputDataView, FConstStructView(RowHandle.DataTable->RowStruct, RowMemory)))
+	{
+		return EDataLinkExecutionReply::Unhandled;
+	}
+
+	InExecutor.Next(this);
 	return EDataLinkExecutionReply::Handled;
 }
diff --git a/Engine/Plugins/Experimental/DataLink/Source/DataLinkEdGraph/Private/Blueprint/K2Node_AsyncDataLinkRequest.cpp b/Engine/Plugins/Experimental/DataLink/Source/DataLinkEdGraph/Private/Blueprint/K2Node_AsyncDataLinkRequest.cpp
--- a/Engine/Plugins/Experimental/DataLink/Source/DataLinkEdGraph/Private/Blueprint/K2Node_AsyncDataLinkRequest.cpp
+++ b/Engine/Plugins/Experimental/DataLink/Source/DataLinkEdGraph/Private/Blueprint/K2Node_AsyncDataLinkRequest.cpp
@@ -35,9 +35,9 @@ UEdGraphPin* UK2Node_AsyncDataLinkRequest::GetExecutionResultPin() const
 	return FindPinChecked(TEXT("ExecutionResult"), EGPD_Output);
 }
 
-UEdGraphPin* UK2Node_AsyncDataLinkRequest::GetRequestCompletePin() const
+UEdGraphPin* UK2Node_AsyncDataLinkRequest::GetOnOutputDataPin() const
 {
-	return FindPinChecked(TEXT("OnRequestComplete"), EGPD_Output);
+	return FindPinChecked(TEXT("OnOutputData"), EGPD_Output);
 }
 
 void UK2Node_AsyncDataLinkRequest::GetMenuActions(FBlueprintActionDatabaseRegistrar& InActionRegistrar) const
diff --git a/Engine/Plugins/Experimental/DataLink/Source/DataLinkEdGraph/Private/Blueprint/K2Node_AsyncDataLinkRequest.h b/Engine/Plugins/Experimental/DataLink/Source/DataLinkEdGraph/Private/Blueprint/K2Node_AsyncDataLinkRequest.h
--- a/Engine/Plugins/Experimental/DataLink/Source/DataLinkEdGraph/Private/Blueprint/K2Node_AsyncDataLinkRequest.h
+++ b/Engine/Plugins/Experimental/DataLink/Source/DataLinkEdGraph/Private/Blueprint/K2Node_AsyncDataLinkRequest.h
@@ -18,7 +18,7 @@ public:
 	UEdGraphPin* GetDataLinkSinkProviderPin() const;
 	UEdGraphPin* GetOutputDataPin() const;
 	UEdGraphPin* GetExecutionResultPin() const;
-	UEdGraphPin* GetRequestCompletePin() const;
+	UEdGraphPin* GetOnOutputDataPin() const;
 
 protected:
 	//~ Begin UK2Node
diff --git a/Engine/Plugins/Experimental/DataLink/Source/DataLinkEdGraph/Private/Blueprint/K2Node_DataLinkRequest.cpp b/Engine/Plugins/Experimental/DataLink/Source/DataLinkEdGraph/Private/Blueprint/K2Node_DataLinkRequest.cpp
--- a/Engine/Plugins/Experimental/DataLink/Source/DataLinkEdGraph/Private/Blueprint/K2Node_DataLinkRequest.cpp
+++ b/Engine/Plugins/Experimental/DataLink/Source/DataLinkEdGraph/Private/Blueprint/K2Node_DataLinkRequest.cpp
@@ -50,18 +50,20 @@ void UK2Node_DataLinkRequest::AllocateDefaultPins()
 
 	// Then Pin
 	UEdGraphPin* const ThenPin = CreatePin(EGPD_Output, UEdGraphSchema_K2::PC_Exec, UEdGraphSchema_K2::PN_Then);
-	ThenPin->PinFriendlyName = LOCTEXT("ThenPinLabel", "Processing");
 	ThenPin->PinToolTip = TEXT("For execution that needs to happen immediately without waiting for a response.");
 
 	// Success Pin
 	UEdGraphPin* const SuccessPin = CreatePin(EGPD_Output, UEdGraphSchema_K2::PC_Exec, UK2Node_DataLinkRequest::PN_Success);
-	SuccessPin->PinFriendlyName = LOCTEXT("SuccessPinLabel", "Response Received");
-	SuccessPin->PinToolTip = TEXT("Executes only when the request was successful and received the data");
+	SuccessPin->PinFriendlyName = LOCTEXT("SuccessPinLabel", "Output Data");
+	SuccessPin->PinToolTip = TEXT("Called when the execution received and processed data successfully");
+
+	// Create Output Pin
+	CreatePin(EGPD_Output, UEdGraphSchema_K2::PC_Struct, FInstancedStruct::StaticStruct(), UDataLinkEdNode::PN_Output);
 
 	// Failure pin
 	UEdGraphPin* const FailurePin = CreatePin(EGPD_Output, UEdGraphSchema_K2::PC_Exec, UK2Node_DataLinkRequest::PN_Failure);
-	FailurePin->PinFriendlyName = LOCTEXT("FailurePinLabel", "Request Failed");
-	FailurePin->PinToolTip = TEXT("Executes when the request failed and no data will be received");
+	FailurePin->PinFriendlyName = LOCTEXT("FailurePinLabel", "Data Link Failed");
+	FailurePin->PinToolTip = TEXT("Called when the execution failed and no data will be received");
 
 	// Create Data Link Instance Pin
 	CreatePin(EGPD_Input, UEdGraphSchema_K2::PC_Struct, FDataLinkInstance::StaticStruct(), UK2Node_DataLinkRequest::PN_DataLinkInstance);
@@ -71,6 +73,7 @@ void UK2Node_DataLinkRequest::AllocateDefaultPins()
 	{
 		// Prevent user from editing default value (i.e. hiding the object dropdown as option).
 		ExecutionContextPin->bDefaultValueIsIgnored = true;
+		ExecutionContextPin->bAdvancedView = true;
 	}
 
 	// Create Data Link Sink Provider Pin
@@ -78,10 +81,13 @@ void UK2Node_DataLinkRequest::AllocateDefaultPins()
 	{
 		// Prevent user from editing default value (i.e. hiding the object dropdown as option).
 		SinkProviderPin->bDefaultValueIsIgnored = true;
+		SinkProviderPin->bAdvancedView = true;
 	}
 
-	// Create Output Pin
-	CreatePin(EGPD_Output, UEdGraphSchema_K2::PC_Struct, FInstancedStruct::StaticStruct(), UDataLinkEdNode::PN_Output);
+	if (AdvancedPinDisplay == ENodeAdvancedPins::NoPins)
+	{
+		AdvancedPinDisplay = ENodeAdvancedPins::Hidden;
+	}
 }
 
 bool UK2Node_DataLinkRequest::IsCompatibleWithGraph(const UEdGraph* InTargetGraph) const
@@ -98,7 +104,7 @@ bool UK2Node_DataLinkRequest::IsCompatibleWithGraph(const UEdGraph* InTargetGrap
 
 FText UK2Node_DataLinkRequest::GetNodeTitle(ENodeTitleType::Type InTitleType) const
 {
-	return LOCTEXT("NodeTitle", "Data Link Request");
+	return LOCTEXT("NodeTitle", "Execute Data Link");
 }
 
 void UK2Node_DataLinkRequest::ExpandNode(FKismetCompilerContext& InCompilerContext, UEdGraph* InSourceGraph)
@@ -119,13 +125,13 @@ void UK2Node_DataLinkRequest::ExpandNode(FKismetCompilerContext& InCompilerConte
 	SwitchNode->Enum = ResultEnum;
 	SwitchNode->AllocateDefaultPins();
 
-	// Connect Request and Result Enum value pins to the Switch Node
+	// Connect Output Data and Result Enum value pins to the Switch Node
 	{
-		UEdGraphPin* RequestCompletePin = AsyncRequestNode->GetRequestCompletePin();
+		UEdGraphPin* OnOutputDataPin = AsyncRequestNode->GetOnOutputDataPin();
 		UEdGraphPin* ExecutionResultPin = AsyncRequestNode->GetExecutionResultPin();
-		check(RequestCompletePin && ExecutionResultPin);
+		check(OnOutputDataPin && ExecutionResultPin);
 
-		RequestCompletePin->MakeLinkTo(SwitchNode->GetExecPin());
+		OnOutputDataPin->MakeLinkTo(SwitchNode->GetExecPin());
 		ExecutionResultPin->MakeLinkTo(SwitchNode->GetSelectionPin());
 	}
 
diff --git a/Engine/Plugins/Experimental/DataLink/Source/DataLinkEdGraph/Private/Nodes/DataLinkEdOutputNode.cpp b/Engine/Plugins/Experimental/DataLink/Source/DataLinkEdGraph/Private/Nodes/DataLinkEdOutputNode.cpp
--- a/Engine/Plugins/Experimental/DataLink/Source/DataLinkEdGraph/Private/Nodes/DataLinkEdOutputNode.cpp
+++ b/Engine/Plugins/Experimental/DataLink/Source/DataLinkEdGraph/Private/Nodes/DataLinkEdOutputNode.cpp
@@ -3,6 +3,11 @@
 #include "Nodes/DataLinkEdOutputNode.h"
 #include "DataLinkEdGraphSchema.h"
 
+UEdGraphPin* UDataLinkEdOutputNode::GetOutputResultPin() const
+{
+	return FindPin(UDataLinkEdNode::PN_Output, EEdGraphPinDirection::EGPD_Input);
+}
+
 FLinearColor UDataLinkEdOutputNode::GetNodeTitleColor() const
 {
 	return FLinearColor::Red;
diff --git a/Engine/Plugins/Experimental/DataLink/Source/DataLinkEdGraph/Public/Nodes/DataLinkEdOutputNode.h b/Engine/Plugins/Experimental/DataLink/Source/DataLinkEdGraph/Public/Nodes/DataLinkEdOutputNode.h
--- a/Engine/Plugins/Experimental/DataLink/Source/DataLinkEdGraph/Public/Nodes/DataLinkEdOutputNode.h
+++ b/Engine/Plugins/Experimental/DataLink/Source/DataLinkEdGraph/Public/Nodes/DataLinkEdOutputNode.h
@@ -16,6 +16,9 @@ class UDataLinkEdOutputNode : public UDataLinkEdNode
 	GENERATED_BODY()
 
 public:
+	/** Retrieves the output result pin */
+	DATALINKEDGRAPH_API UEdGraphPin* GetOutputResultPin() const;
+
 	//~ Begin UDataLinkEdNode
 	virtual bool RequiresPinRecreation() const override { return false; }
 	//~ End UDataLinkEdNode
diff --git a/Engine/Plugins/Experimental/DataLink/Source/DataLinkEditor/Private/Compiler/DataLinkGraphCompiler.cpp b/Engine/Plugins/Experimental/DataLink/Source/DataLinkEditor/Private/Compiler/DataLinkGraphCompiler.cpp
--- a/Engine/Plugins/Experimental/DataLink/Source/DataLinkEditor/Private/Compiler/DataLinkGraphCompiler.cpp
+++ b/Engine/Plugins/Experimental/DataLink/Source/DataLinkEditor/Private/Compiler/DataLinkGraphCompiler.cpp
@@ -74,7 +74,7 @@ bool FDataLinkGraphCompiler::CompileNodes()
 
 bool FDataLinkGraphCompiler::CreateCompiledNodes()
 {
-	UDataLinkEdOutputNode* OutputNode = DataLinkEdGraph->FindOutputNode();
+	UDataLinkEdOutputNode* const OutputNode = DataLinkEdGraph->FindOutputNode();
 	if (!OutputNode)
 	{
 		UE_LOG(LogDataLinkEditor, Error, TEXT("Compilation failed. Output Node was not valid! Graph: %s")
@@ -82,9 +82,34 @@ bool FDataLinkGraphCompiler::CreateCompiledNodes()
 		return false;
 	}
 
+	const UEdGraphPin* OutputResultPin = OutputNode->GetOutputResultPin();
+	if (!OutputResultPin || OutputResultPin->LinkedTo.IsEmpty() || !OutputResultPin->LinkedTo[0])
+	{
+		UE_LOG(LogDataLinkEditor, Error, TEXT("Compilation failed. No output result pin was found or connected in graph '%s'"), *DataLinkGraph->GetName());
+		return false;
+	}
+
+	// The node that the cosmetic output node is connected to is the actual output node at runtime
+	const UEdGraphPin* OutputPin = OutputResultPin->LinkedTo[0];
+	if (!ensure(OutputPin))
+	{
+		UE_LOG(LogDataLinkEditor, Error, TEXT("Compilation failed. Output pin connected to the output result node was invalid in graph '%s'"), *DataLinkGraph->GetName());
+		return false;
+	}
+
+	OutputEdNode = Cast<UDataLinkEdNode>(OutputPin->GetOwningNode());
+	OutputPinName = OutputPin->PinName;
+
+	if (!ensure(OutputEdNode))
+	{
+		UE_LOG(LogDataLinkEditor, Error, TEXT("Compilation failed. Output node connected to the output result node was not a data link node in graph '%s'"), *DataLinkGraph->GetName());
+		return false;
+	}
+
 	EdToCompiledMap.Empty(DataLinkEdGraph->Nodes.Num());
 
-	TArray<const UDataLinkEdNode*> EdNodesRemaining;
+	TArray<const UDataLinkEdNode*, TInlineAllocator<1>> EdNodesRemaining;
+	EdNodesRemaining.Add(OutputEdNode);
 
 	const auto AddLinkedNodes =
 		[&EdToCompiledMap = EdToCompiledMap, &EdNodesRemaining](const UEdGraphPin& InPin, const UDataLinkEdNode& InLinkedNode, const UEdGraphPin& InLinkedPin)
@@ -98,17 +123,6 @@ bool FDataLinkGraphCompiler::CreateCompiledNodes()
 			}
 		};
 
-	OutputNode->ForEachPinConnection(AddLinkedNodes);
-
-	if (EdNodesRemaining.IsEmpty())
-	{
-		UE_LOG(LogDataLinkEditor, Error, TEXT("Compilation failed. No nodes were considered for compilation in graph '%s'"), *DataLinkGraph->GetName());
-		return false;
-	}
-
-	// The node that the cosmetic output node is connected to is the actual output node at runtime
-	OutputEdNode = EdNodesRemaining[0];
-
 	while (!EdNodesRemaining.IsEmpty())
 	{
 		const UDataLinkEdNode* EdNode = EdNodesRemaining.Pop();
@@ -235,6 +249,7 @@ bool FDataLinkGraphCompiler::SetInputOutputNodes()
 	}
 
 	DataLinkGraph->OutputNode = OutputNode;
+	DataLinkGraph->OutputPinName = OutputPinName;
 
 	// Recursively adds all the graph's entry nodes (i.e. nodes with input pins that are not connected to other nodes)
 	AddGraphEntryNodes(OutputNode);
diff --git a/Engine/Plugins/Experimental/DataLink/Source/DataLinkEditor/Private/Compiler/DataLinkGraphCompiler.h b/Engine/Plugins/Experimental/DataLink/Source/DataLinkEditor/Private/Compiler/DataLinkGraphCompiler.h
--- a/Engine/Plugins/Experimental/DataLink/Source/DataLinkEditor/Private/Compiler/DataLinkGraphCompiler.h
+++ b/Engine/Plugins/Experimental/DataLink/Source/DataLinkEditor/Private/Compiler/DataLinkGraphCompiler.h
@@ -39,6 +39,12 @@ private:
 
 	TObjectPtr<const UDataLinkEdNode> OutputEdNode;
 
+	/**
+	 * Pin name of the output ed node connecting to the cosmetic output node
+	 * Indicating that this pin is the pin that will be containing the graph output data
+	 */
+	FName OutputPinName;
+
 	/** Map of the Editor Node to its Compiled Node */
 	TMap<const UDataLinkEdNode*, UDataLinkNode*> EdToCompiledMap;
 };
diff --git a/Engine/Plugins/Experimental/DataLink/Source/DataLinkEditor/Private/DataLinkEditorStyle.cpp b/Engine/Plugins/Experimental/DataLink/Source/DataLinkEditor/Private/DataLinkEditorStyle.cpp
--- a/Engine/Plugins/Experimental/DataLink/Source/DataLinkEditor/Private/DataLinkEditorStyle.cpp
+++ b/Engine/Plugins/Experimental/DataLink/Source/DataLinkEditor/Private/DataLinkEditorStyle.cpp
@@ -23,7 +23,8 @@ FDataLinkEditorStyle::FDataLinkEditorStyle()
 	CoreContentRootDir = FPaths::EngineContentDir() / TEXT("Slate");
 
 	// Editor Commands
-	Set("DataLinkGraph.Preview", new IMAGE_BRUSH_SVG("Starship/MainToolbar/simulate", Icon20));
+	Set("DataLinkGraph.RunPreview", new IMAGE_BRUSH_SVG("Starship/MainToolbar/simulate", Icon20));
+	Set("DataLinkGraph.StopPreview", new CORE_IMAGE_BRUSH_SVG("Starship/Common/stop", Icon20));
 	Set("DataLinkGraph.ClearPreviewOutput", new IMAGE_BRUSH_SVG("Starship/Common/ResetToDefault", Icon20));
 	Set("DataLinkGraph.ClearPreviewCache", new IMAGE_BRUSH("Icons/GeneralTools/Delete_40x", Icon20));
 
diff --git a/Engine/Plugins/Experimental/DataLink/Source/DataLinkEditor/Private/DataLinkGraphCommands.cpp b/Engine/Plugins/Experimental/DataLink/Source/DataLinkEditor/Private/DataLinkGraphCommands.cpp
--- a/Engine/Plugins/Experimental/DataLink/Source/DataLinkEditor/Private/DataLinkGraphCommands.cpp
+++ b/Engine/Plugins/Experimental/DataLink/Source/DataLinkEditor/Private/DataLinkGraphCommands.cpp
@@ -22,12 +22,18 @@ void FDataLinkGraphCommands::RegisterCommands()
 		, EUserInterfaceActionType::Button
 		, FInputChord());
 
-	UI_COMMAND(Preview
+	UI_COMMAND(RunPreview
 		, "Preview"
 		, "Triggers execution of the data link graph with the given input data"
 		, EUserInterfaceActionType::Button
 		, FInputChord());
 
+	UI_COMMAND(StopPreview
+		, "Stop"
+		, "Stops the active preview execution"
+		, EUserInterfaceActionType::Button
+		, FInputChord());
+
 	UI_COMMAND(ClearPreviewOutput
 		, "Clear Output"
 		, "Clears the output data from the last execution."
diff --git a/Engine/Plugins/Experimental/DataLink/Source/DataLinkEditor/Private/DataLinkGraphCommands.h b/Engine/Plugins/Experimental/DataLink/Source/DataLinkEditor/Private/DataLinkGraphCommands.h
--- a/Engine/Plugins/Experimental/DataLink/Source/DataLinkEditor/Private/DataLinkGraphCommands.h
+++ b/Engine/Plugins/Experimental/DataLink/Source/DataLinkEditor/Private/DataLinkGraphCommands.h
@@ -15,7 +15,9 @@ public:
 
 	TSharedPtr<FUICommandInfo> Compile;
 
-	TSharedPtr<FUICommandInfo> Preview;
+	TSharedPtr<FUICommandInfo> RunPreview;
+
+	TSharedPtr<FUICommandInfo> StopPreview;
 
 	TSharedPtr<FUICommandInfo> ClearPreviewOutput;
 
diff --git a/Engine/Plugins/Experimental/DataLink/Source/DataLinkEditor/Private/Preview/DataLinkPreviewTool.cpp b/Engine/Plugins/Experimental/DataLink/Source/DataLinkEditor/Private/Preview/DataLinkPreviewTool.cpp
--- a/Engine/Plugins/Experimental/DataLink/Source/DataLinkEditor/Private/Preview/DataLinkPreviewTool.cpp
+++ b/Engine/Plugins/Experimental/DataLink/Source/DataLinkEditor/Private/Preview/DataLinkPreviewTool.cpp
@@ -40,10 +40,14 @@ void FDataLinkPreviewTool::BindCommands(const TSharedRef<FUICommandList>& InComm
 {
 	const FDataLinkGraphCommands& GraphCommands = FDataLinkGraphCommands::Get();
 
-	InCommandList->MapAction(GraphCommands.Preview
+	InCommandList->MapAction(GraphCommands.RunPreview
 		, FExecuteAction::CreateSP(this, &FDataLinkPreviewTool::RunPreview)
 		, FCanExecuteAction::CreateSP(this, &FDataLinkPreviewTool::CanRunPreview));
 
+	InCommandList->MapAction(GraphCommands.StopPreview
+		, FExecuteAction::CreateSP(this, &FDataLinkPreviewTool::StopPreview)
+		, FCanExecuteAction::CreateSP(this, &FDataLinkPreviewTool::CanStopPreview));
+
 	InCommandList->MapAction(GraphCommands.ClearPreviewOutput
 		, FExecuteAction::CreateSP(this, &FDataLinkPreviewTool::ClearOutput)
 		, FCanExecuteAction::CreateSP(this, &FDataLinkPreviewTool::CanClearOutput));
@@ -86,12 +90,12 @@ void FDataLinkPreviewTool::AddReferencedObjects(FReferenceCollector& InCollector
 bool FDataLinkPreviewTool::CanRunPreview() const
 {
 	// Allow execution if there's not already one taking place
-	return !Executor.IsValid();
+	return !Executor.IsValid() || !Executor->IsRunning();
 }
 
 void FDataLinkPreviewTool::RunPreview()
 {
-	if (Executor.IsValid())
+	if (!CanRunPreview())
 	{
 		UE_LOG(LogDataLinkEditor, Error, TEXT("[%s] Data Link execution is in progress!"), Executor->GetContextName().GetData());
 		return;
@@ -108,11 +112,26 @@ void FDataLinkPreviewTool::RunPreview()
 #endif
 		.SetContextObject(AssetEditor)
 		.SetSink(Sink)
-		.SetOnFinish(FOnDataLinkExecutionFinished::CreateSP(this, &FDataLinkPreviewTool::OnPreviewFinished)));
+		.SetOnOutputData(FOnDataLinkOutputData::CreateSP(this, &FDataLinkPreviewTool::OnPreviewOutputData)));
 
 	Executor->Run();
 }
 
+bool FDataLinkPreviewTool::CanStopPreview() const
+{
+	// Can only stop preview if there's an active executor
+	return Executor.IsValid() && Executor->IsRunning();
+}
+
+void FDataLinkPreviewTool::StopPreview()
+{
+	if (CanStopPreview())
+	{
+		Executor->Stop();
+		Executor.Reset();
+	}
+}
+
 bool FDataLinkPreviewTool::CanClearOutput() const
 {
 	return PreviewData && PreviewData->OutputData.IsValid();
@@ -139,18 +158,20 @@ void FDataLinkPreviewTool::ClearCache()
 	Sink.Reset();
 }
 
-void FDataLinkPreviewTool::OnPreviewFinished(const FDataLinkExecutor& InExecutor, EDataLinkExecutionResult InResult, FConstStructView InOutputDataView)
+void FDataLinkPreviewTool::OnPreviewOutputData(const FDataLinkExecutor& InExecutor, FConstStructView InOutputDataView)
 {
-	Executor.Reset();
-
-	const UScriptStruct* OutputDataStruct = InOutputDataView.GetScriptStruct();
-	if (!OutputDataStruct || InResult == EDataLinkExecutionResult::Failed)
+	if (!InOutputDataView.IsValid())
 	{
 		PreviewData->OutputData.Reset();
 		return;
 	}
 
-	PreviewData->OutputData.InitializeAs(OutputDataStruct);
+	const UScriptStruct* OutputDataStruct = InOutputDataView.GetScriptStruct();
+	if (PreviewData->OutputData.GetScriptStruct() != OutputDataStruct)
+	{
+		PreviewData->OutputData.InitializeAs(OutputDataStruct);
+	}
+
 	OutputDataStruct->CopyScriptStruct(PreviewData->OutputData.GetMutableMemory(), InOutputDataView.GetMemory());
 	PreviewDetails->ForceRefresh();
 }
@@ -172,7 +193,8 @@ void FDataLinkPreviewTool::RegisterToolbar()
 	const FDataLinkGraphCommands& GraphCommands = FDataLinkGraphCommands::Get();
 
 	FToolMenuSection& PreviewSection = ToolbarMenu->FindOrAddSection(UE::DataLinkEditor::PreviewSectionName);
-	PreviewSection.AddEntry(FToolMenuEntry::InitToolBarButton(GraphCommands.Preview));
+	PreviewSection.AddEntry(FToolMenuEntry::InitToolBarButton(GraphCommands.RunPreview));
+	PreviewSection.AddEntry(FToolMenuEntry::InitToolBarButton(GraphCommands.StopPreview));
 	PreviewSection.AddEntry(FToolMenuEntry::InitToolBarButton(GraphCommands.ClearPreviewOutput));
 	PreviewSection.AddEntry(FToolMenuEntry::InitToolBarButton(GraphCommands.ClearPreviewCache));
 }
diff --git a/Engine/Plugins/Experimental/DataLink/Source/DataLinkEditor/Private/Preview/DataLinkPreviewTool.h b/Engine/Plugins/Experimental/DataLink/Source/DataLinkEditor/Private/Preview/DataLinkPreviewTool.h
--- a/Engine/Plugins/Experimental/DataLink/Source/DataLinkEditor/Private/Preview/DataLinkPreviewTool.h
+++ b/Engine/Plugins/Experimental/DataLink/Source/DataLinkEditor/Private/Preview/DataLinkPreviewTool.h
@@ -52,13 +52,16 @@ private:
 	bool CanRunPreview() const;
 	void RunPreview();
 
+	bool CanStopPreview() const;
+	void StopPreview();
+
 	bool CanClearOutput() const;
 	void ClearOutput();
 
 	bool CanClearCache() const;
 	void ClearCache();
 
-	void OnPreviewFinished(const FDataLinkExecutor& InExecutor, EDataLinkExecutionResult InResult, FConstStructView InOutputDataView);
+	void OnPreviewOutputData(const FDataLinkExecutor& InExecutor, FConstStructView InOutputDataView);
 
 	void RegisterToolbar();
 
diff --git a/Engine/Plugins/Experimental/DataLink/Source/DataLinkHttp/Private/DataLinkHttpSource.cpp b/Engine/Plugins/Experimental/DataLink/Source/DataLinkHttp/Private/DataLinkHttpSource.cpp
--- a/Engine/Plugins/Experimental/DataLink/Source/DataLinkHttp/Private/DataLinkHttpSource.cpp
+++ b/Engine/Plugins/Experimental/DataLink/Source/DataLinkHttp/Private/DataLinkHttpSource.cpp
@@ -68,11 +68,11 @@ EDataLinkExecutionReply UDataLinkHttpSource::OnExecute(FDataLinkExecutor& InExec
 				FDataLinkString& OutputDataView = OutputDataViewer.Get<FDataLinkString>(UE::DataLink::OutputDefault);
 				OutputDataView.Value = InResponse->GetContentAsString();
 
-				Executor->SucceedNode(this, FConstStructView::Make(OutputDataView));
+				Executor->Next(this);
 			}
 			else
 			{
-				Executor->FailNode(this);
+				Executor->Fail(this);
 			}
 		});
 
diff --git a/Engine/Plugins/Experimental/DataLink/Source/DataLinkHttp/Private/DataLinkNodeHttpSettingsBuilder.cpp b/Engine/Plugins/Experimental/DataLink/Source/DataLinkHttp/Private/DataLinkNodeHttpSettingsBuilder.cpp
--- a/Engine/Plugins/Experimental/DataLink/Source/DataLinkHttp/Private/DataLinkNodeHttpSettingsBuilder.cpp
+++ b/Engine/Plugins/Experimental/DataLink/Source/DataLinkHttp/Private/DataLinkNodeHttpSettingsBuilder.cpp
@@ -31,7 +31,7 @@ EDataLinkExecutionReply UDataLinkNodeHttpSettingsBuilder::OnExecute(FDataLinkExe
 	const FDataLinkStringBuilder URLBuilder(URLSegments, Tokens);
 	if (!URLBuilder.BuildString(NodeInstance.GetInputDataViewer(), HttpSettings.URL))
 	{
-		InExecutor.FailNode(this);
+		InExecutor.Fail(this);
 		return EDataLinkExecutionReply::Handled;
 	}
 
@@ -39,7 +39,7 @@ EDataLinkExecutionReply UDataLinkNodeHttpSettingsBuilder::OnExecute(FDataLinkExe
 	HttpSettings.Headers = Headers;
 	HttpSettings.Body = Body;
 
-	InExecutor.SucceedNode(this, FConstStructView::Make(HttpSettings));
+	InExecutor.Next(this);
 	return EDataLinkExecutionReply::Handled;
 }
 
diff --git a/Engine/Plugins/Experimental/DataLink/Source/DataLinkJson/DataLinkJson.Build.cs b/Engine/Plugins/Experimental/DataLink/Source/DataLinkJson/DataLinkJson.Build.cs
--- a/Engine/Plugins/Experimental/DataLink/Source/DataLinkJson/DataLinkJson.Build.cs
+++ b/Engine/Plugins/Experimental/DataLink/Source/DataLinkJson/DataLinkJson.Build.cs
@@ -21,6 +21,7 @@ public class DataLinkJson : ModuleRules
             new string[]
             {
                 "Engine",
+                "Json",
                 "JsonUtilities",
             }
         );
diff --git a/Engine/Plugins/Experimental/DataLink/Source/DataLinkJson/Private/DataLinkJsonFromStruct.cpp b/Engine/Plugins/Experimental/DataLink/Source/DataLinkJson/Private/DataLinkJsonFromStruct.cpp
--- a/Engine/Plugins/Experimental/DataLink/Source/DataLinkJson/Private/DataLinkJsonFromStruct.cpp
+++ b/Engine/Plugins/Experimental/DataLink/Source/DataLinkJson/Private/DataLinkJsonFromStruct.cpp
@@ -52,7 +52,7 @@ EDataLinkExecutionReply UDataLinkJsonFromStruct::OnExecute(FDataLinkExecutor& In
 	// Set the json object wrapper's string to the latest json object result
 	Wrapper.JsonObjectToString(Wrapper.JsonString);
 
-	InExecutor.SucceedNode(this, FConstStructView::Make(Wrapper));
+	InExecutor.Next(this);
 	return EDataLinkExecutionReply::Handled;
 }
 
diff --git a/Engine/Plugins/Experimental/DataLink/Source/DataLinkJson/Private/DataLinkJsonToStruct.cpp b/Engine/Plugins/Experimental/DataLink/Source/DataLinkJson/Private/DataLinkJsonToStruct.cpp
--- a/Engine/Plugins/Experimental/DataLink/Source/DataLinkJson/Private/DataLinkJsonToStruct.cpp
+++ b/Engine/Plugins/Experimental/DataLink/Source/DataLinkJson/Private/DataLinkJsonToStruct.cpp
@@ -65,7 +65,7 @@ EDataLinkExecutionReply UDataLinkJsonToStruct::OnExecute(FDataLinkExecutor& InEx
 		return EDataLinkExecutionReply::Unhandled;
 	}
 
-	InExecutor.SucceedNode(this, OutputDataView);
+	InExecutor.Next(this);
 	return EDataLinkExecutionReply::Handled;
 }
 
diff --git a/Engine/Plugins/Experimental/DataLink/Source/DataLinkJson/Private/DataLinkStringToJson.cpp b/Engine/Plugins/Experimental/DataLink/Source/DataLinkJson/Private/DataLinkStringToJson.cpp
--- a/Engine/Plugins/Experimental/DataLink/Source/DataLinkJson/Private/DataLinkStringToJson.cpp
+++ b/Engine/Plugins/Experimental/DataLink/Source/DataLinkJson/Private/DataLinkStringToJson.cpp
@@ -9,6 +9,8 @@
 #include "DataLinkNodeInstance.h"
 #include "DataLinkPinBuilder.h"
 #include "JsonObjectWrapper.h"
+#include "Serialization/JsonReader.h"
+#include "Serialization/JsonSerializer.h"
 #include "StructUtils/StructView.h"
 
 #define LOCTEXT_NAMESPACE "DataLinkStringToJson"
@@ -40,13 +42,48 @@ EDataLinkExecutionReply UDataLinkStringToJson::OnExecute(FDataLinkExecutor& InEx
 	// Return early if the output data's json string already matches the input string
 	if (OutputData.JsonString == InputData.Value)
 	{
-		InExecutor.SucceedNode(this, FConstStructView::Make(OutputData));
+		InExecutor.Next(this);
 		return EDataLinkExecutionReply::Handled;
 	}
 
-	if (!OutputData.JsonObjectFromString(InputData.Value))
+	TSharedPtr<FJsonValue> JsonValue;
+
+	// Deserialize JSON
+	{
+		TSharedRef<TJsonReader<>> JsonReader = TJsonReaderFactory<>::Create(InputData.Value);
+		if (!FJsonSerializer::Deserialize(JsonReader, JsonValue))
+		{
+			UE_LOG(LogDataLinkJson, Error, TEXT("[%.*s] Failed to deserialize JSON from string %s")
+				, InExecutor.GetContextName().Len()
+				, InExecutor.GetContextName().GetData()
+				, *InputData.Value);
+			return EDataLinkExecutionReply::Unhandled;
+		}
+	}
+
+	const TArray<TSharedPtr<FJsonValue>>* JsonArray;
+	if (JsonValue->TryGetArray(JsonArray) && ensure(JsonArray))
+	{
+		if (JsonArray->Num() != 1 || !(*JsonArray)[0].IsValid())
+		{
+			UE_LOG(LogDataLinkJson, Error, TEXT("[%.*s] Failed to find a valid JSON object from string %s. Expected 1 valid json object in top level array. Found %d")
+				, InExecutor.GetContextName().Len()
+				, InExecutor.GetContextName().GetData()
+				, *InputData.Value
+				, JsonArray->Num());
+			return EDataLinkExecutionReply::Unhandled;
+		}
+
+		OutputData.JsonObject = (*JsonArray)[0]->AsObject();
+	}
+	else
+	{
+		OutputData.JsonObject = JsonValue->AsObject();
+	}
+
+	if (!OutputData.JsonObject.IsValid())
 	{
-		UE_LOG(LogDataLinkJson, Error, TEXT("[%.*s] Failed to create a JSON object from string %s")
+		UE_LOG(LogDataLinkJson, Error, TEXT("[%.*s] Failed to find a valid JSON object from string %s")
 			, InExecutor.GetContextName().Len()
 			, InExecutor.GetContextName().GetData()
 			, *InputData.Value);
@@ -55,7 +92,7 @@ EDataLinkExecutionReply UDataLinkStringToJson::OnExecute(FDataLinkExecutor& InEx
 
 	// Save the Json String for future reference if the parsed object can be re-used for matching strings
 	OutputData.JsonString = InputData.Value;
-	InExecutor.SucceedNode(this, FConstStructView::Make(OutputData));
+	InExecutor.Next(this);
 	return EDataLinkExecutionReply::Handled;
 }
 
diff --git a/Engine/Plugins/Experimental/DataLinkOAuth/Source/DataLinkOAuth/Private/DataLinkNodeOAuth.cpp b/Engine/Plugins/Experimental/DataLinkOAuth/Source/DataLinkOAuth/Private/DataLinkNodeOAuth.cpp
--- a/Engine/Plugins/Experimental/DataLinkOAuth/Source/DataLinkOAuth/Private/DataLinkNodeOAuth.cpp
+++ b/Engine/Plugins/Experimental/DataLinkOAuth/Source/DataLinkOAuth/Private/DataLinkNodeOAuth.cpp
@@ -65,7 +65,7 @@ EDataLinkExecutionReply UDataLinkNodeOAuth::OnExecute(FDataLinkExecutor& InExecu
 
 			if (OAuthSettings->AuthorizeHttpRequest(*OAuthToken, OutputHttpSettings))
 			{
-				InExecutor.SucceedNode(this, FConstStructView::Make(OutputHttpSettings));
+				InExecutor.Next(this);
 				return EDataLinkExecutionReply::Handled;
 			}
 		}
@@ -121,7 +121,7 @@ void UDataLinkNodeOAuth::OnAuthResponse(const FHttpServerRequest& InRequest, TWe
 			, Executor->GetContextName().GetData()
 			, *OAuthSettings->GetName());
 
-		Executor->FailNode(this);
+		Executor->Fail(this);
 		return;
 	}
 
@@ -137,7 +137,7 @@ void UDataLinkNodeOAuth::OnAuthResponse(const FHttpServerRequest& InRequest, TWe
 			, Executor->GetContextName().GetData()
 			, *OAuthSettings->GetName());
 
-		Executor->FailNode(this);
+		Executor->Fail(this);
 	}
 }
 
@@ -162,7 +162,7 @@ void UDataLinkNodeOAuth::OnExchangeCodeResponse(FStringView InResponse, TWeakPtr
 			, Executor->GetContextName().GetData()
 			, *OAuthSettings->GetName());
 
-		Executor->FailNode(this);
+		Executor->Fail(this);
 		return;
 	}
 
@@ -174,12 +174,12 @@ void UDataLinkNodeOAuth::OnExchangeCodeResponse(FStringView InResponse, TWeakPtr
 			, *OAuthSettings->GetName()
 			, InResponse.GetData());
 
-		Executor->FailNode(this);
+		Executor->Fail(this);
 		return;
 	}
 
 	const FDataLinkOutputDataViewer& OutputDataViewer = NodeInstance.GetOutputDataViewer();
-	
+
 	// Copy Input Http Settings into Output
 	FDataLinkHttpSettings& OutputHttpSettings = OutputDataViewer.Get<FDataLinkHttpSettings>(UE::DataLink::OutputDefault);
 	OutputHttpSettings = InputDataViewer.Get<FDataLinkHttpSettings>(UE::DataLinkOAuth::InputHttp);
@@ -191,7 +191,7 @@ void UDataLinkNodeOAuth::OnExchangeCodeResponse(FStringView InResponse, TWeakPtr
 			, *OAuthSettings->GetName()
 			, InResponse.GetData());
 
-		Executor->FailNode(this);
+		Executor->Fail(this);
 		return;
 	}
 
@@ -200,7 +200,7 @@ void UDataLinkNodeOAuth::OnExchangeCodeResponse(FStringView InResponse, TWeakPtr
 		OAuthSubsystem->RegisterToken(OAuthSettings, OAuthToken);
 	}
 
-	Executor->SucceedNode(this, FConstStructView::Make(OutputHttpSettings));
+	Executor->Next(this);
 }
 
 #undef LOCTEXT_NAMESPACE
diff --git a/Engine/Plugins/Experimental/SceneStateDataLink/Source/SceneStateDataLink/Private/SceneStateDataLinkRequestTask.cpp b/Engine/Plugins/Experimental/SceneStateDataLink/Source/SceneStateDataLink/Private/SceneStateDataLinkRequestTask.cpp
--- a/Engine/Plugins/Experimental/SceneStateDataLink/Source/SceneStateDataLink/Private/SceneStateDataLinkRequestTask.cpp
+++ b/Engine/Plugins/Experimental/SceneStateDataLink/Source/SceneStateDataLink/Private/SceneStateDataLinkRequestTask.cpp
@@ -36,27 +36,41 @@ void FSceneStateDataLinkRequestTask::OnStart(const FSceneStateExecutionContext&
 {
 	FInstanceDataType& Instance = InTaskInstance.Get<FInstanceDataType>();
 
+	const UE::SceneState::FTaskExecutionContext TaskExecutionContext(*this, InContext);
+
 	Instance.Executor = FDataLinkExecutor::Create(FDataLinkExecutorArguments(Instance.CreateDataLinkInstance())
 #if WITH_DATALINK_CONTEXT
 		.SetContextName(InContext.GetExecutionContextName())
 #endif
 		.SetContextObject(InContext.GetContextObject())
-		.SetOnFinish(FOnDataLinkExecutionFinished::CreateStatic(&FSceneStateDataLinkRequestTask::OnExecutionFinished, UE::SceneState::FTaskExecutionContext(*this, InContext))));
+		.SetOnOutputData(FOnDataLinkOutputData::CreateStatic(&FSceneStateDataLinkRequestTask::OnOutputData, TaskExecutionContext))
+		.SetOnFinished(FOnDataLinkExecutionFinished::CreateStatic(&FSceneStateDataLinkRequestTask::OnFinished, TaskExecutionContext)));
 
 	Instance.Executor->Run();
 }
 
-void FSceneStateDataLinkRequestTask::OnExecutionFinished(const FDataLinkExecutor& InExecutor, EDataLinkExecutionResult InResult, FConstStructView InOutputDataView, UE::SceneState::FTaskExecutionContext InTaskContext)
+void FSceneStateDataLinkRequestTask::OnStop(const FSceneStateExecutionContext& InContext, FStructView InTaskInstance, ESceneStateTaskStopReason InStopReason) const
+{
+	FInstanceDataType& Instance = InTaskInstance.Get<FInstanceDataType>();
+	if (Instance.Executor.IsValid())
+	{
+		Instance.Executor->Stop();
+		Instance.Executor.Reset();
+	}
+}
+
+void FSceneStateDataLinkRequestTask::OnOutputData(const FDataLinkExecutor& InExecutor, FConstStructView InOutputDataView, UE::SceneState::FTaskExecutionContext InTaskContext)
 {
-	// Ensure that any exit path finishes the task
-	ON_SCOPE_EXIT
+	if (!InOutputDataView.IsValid())
 	{
 		InTaskContext.FinishTask();
-	};
+		return;
+	}
 
 	const FSceneStateExecutionContext* const ExecutionContext = InTaskContext.GetExecutionContext();
 	if (!ExecutionContext)
 	{
+		InTaskContext.FinishTask();
 		return;
 	}
 
@@ -74,6 +88,7 @@ void FSceneStateDataLinkRequestTask::OnExecutionFinished(const FDataLinkExecutor
 	if (!UE::SceneState::ResolveProperty(*ExecutionContext, Instance->OutputTarget, Result))
 	{
 		UE_LOG(LogSceneStateDataLink, Warning, TEXT("[%s] Data Link Request Task failed. Output Target could not be resolved."), *ExecutionContext->GetExecutionContextName());
+		InTaskContext.FinishTask();
 		return;
 	}
 
@@ -86,6 +101,7 @@ void FSceneStateDataLinkRequestTask::OnExecutionFinished(const FDataLinkExecutor
 	{
 		UE_LOG(LogSceneStateDataLink, Error, TEXT("[%s] Data Link Request Task failed. Output Target is not bound to a struct property")
 			, *ExecutionContext->GetExecutionContextName());
+		InTaskContext.FinishTask();
 		return;
 	}
 
@@ -95,8 +111,14 @@ void FSceneStateDataLinkRequestTask::OnExecutionFinished(const FDataLinkExecutor
 			, *ExecutionContext->GetExecutionContextName()
 			, *GetNameSafe(StructProperty->Struct)
 			, *GetNameSafe(InOutputDataView.GetScriptStruct()));
+		InTaskContext.FinishTask();
 		return;
 	}
 
 	StructProperty->Struct->CopyScriptStruct(Result.ValuePtr, InOutputDataView.GetMemory());
 }
+
+void FSceneStateDataLinkRequestTask::OnFinished(const FDataLinkExecutor& InExecutor, EDataLinkExecutionResult InExecutionResult, UE::SceneState::FTaskExecutionContext InTaskContext)
+{
+	InTaskContext.FinishTask();
+}
diff --git a/Engine/Plugins/Experimental/SceneStateDataLink/Source/SceneStateDataLink/Public/SceneStateDataLinkRequestTask.h b/Engine/Plugins/Experimental/SceneStateDataLink/Source/SceneStateDataLink/Public/SceneStateDataLinkRequestTask.h
--- a/Engine/Plugins/Experimental/SceneStateDataLink/Source/SceneStateDataLink/Public/SceneStateDataLinkRequestTask.h
+++ b/Engine/Plugins/Experimental/SceneStateDataLink/Source/SceneStateDataLink/Public/SceneStateDataLinkRequestTask.h
@@ -61,9 +61,13 @@ protected:
 	virtual void OnBuildTaskInstance(UObject* InOuter, FStructView InTaskInstance) const override;
 #endif
 	virtual void OnStart(const FSceneStateExecutionContext& InContext, FStructView InTaskInstance) const override;
+	virtual void OnStop(const FSceneStateExecutionContext& InContext, FStructView InTaskInstance, ESceneStateTaskStopReason InStopReason) const override;
 	//~ End FSceneStateTask
 
 private:
-	/** Called when a data link request has completed */
-	static void OnExecutionFinished(const FDataLinkExecutor& InExecutor, EDataLinkExecutionResult InResult, FConstStructView InOutputDataView, UE::SceneState::FTaskExecutionContext InTaskContext);
+	/** Called when data link has broadcast output data */
+	static void OnOutputData(const FDataLinkExecutor& InExecutor, FConstStructView InOutputDataView, UE::SceneState::FTaskExecutionContext InTaskContext);
+
+	/** Called when a data link execution has finished */
+	static void OnFinished(const FDataLinkExecutor& InExecutor, EDataLinkExecutionResult InExecutionResult, UE::SceneState::FTaskExecutionContext InTaskContext);
 };
