Motion Design Data Link: remove CppStructOps check as GetStructHash64 now checks if a Script Struct has a valid CppStructOps

diff --git a/Engine/Plugins/Experimental/DataLink/Source/DataLink/Private/DataLinkSink.cpp b/Engine/Plugins/Experimental/DataLink/Source/DataLink/Private/DataLinkSink.cpp
--- a/Engine/Plugins/Experimental/DataLink/Source/DataLink/Private/DataLinkSink.cpp
+++ b/Engine/Plugins/Experimental/DataLink/Source/DataLink/Private/DataLinkSink.cpp
@@ -38,8 +38,7 @@ FDataLinkSinkKey::FDataLinkSinkKey(TSubclassOf<UDataLinkNode> InNodeClass, const
 		CachedHash = HashCombineFast(CachedHash, GetTypeHash(InputDataEntry.Name));
 		CachedHash = HashCombineFast(CachedHash, GetTypeHash(InputStruct));
 
-		// GetStructHash64 does not check for whether CppStructOps is valid or not
-		if (InputStruct && InputStruct->GetCppStructOps())
+		if (InputDataEntry.DataView.IsValid())
 		{
 			CachedHash = HashCombineFast(CachedHash, UE::StructUtils::GetStructHash64(InputDataEntry.DataView));
 		}
