Motion Design Data Link: added web socket node implementation

diff --git a/Engine/Plugins/Experimental/DataLink/DataLink.uplugin b/Engine/Plugins/Experimental/DataLink/DataLink.uplugin
--- a/Engine/Plugins/Experimental/DataLink/DataLink.uplugin
+++ b/Engine/Plugins/Experimental/DataLink/DataLink.uplugin
@@ -48,6 +48,11 @@
 			"Name": "DataLinkDataTable",
 			"Type": "Runtime",
 			"LoadingPhase": "Default"
+		},
+		{
+			"Name": "DataLinkWebSocket",
+			"Type": "Runtime",
+			"LoadingPhase": "Default"
 		}
 	]
 }
\ No newline at end of file
diff --git a/Engine/Plugins/Experimental/DataLink/Source/DataLinkWebSocket/DataLinkWebSocket.Build.cs b/Engine/Plugins/Experimental/DataLink/Source/DataLinkWebSocket/DataLinkWebSocket.Build.cs
new file mode 100644
--- /dev/null
+++ b/Engine/Plugins/Experimental/DataLink/Source/DataLinkWebSocket/DataLinkWebSocket.Build.cs
@@ -0,0 +1,28 @@
+// Copyright Epic Games, Inc. All Rights Reserved.
+
+using UnrealBuildTool;
+
+public class DataLinkWebSocket : ModuleRules
+{
+    public DataLinkWebSocket(ReadOnlyTargetRules Target) : base(Target)
+    {
+        PCHUsage = ModuleRules.PCHUsageMode.UseExplicitOrSharedPCHs;
+
+        PublicDependencyModuleNames.AddRange(
+            new string[]
+            {
+                "Core",
+                "CoreUObject",
+                "DataLink",
+                "Engine",
+            }
+        );
+
+        PrivateDependencyModuleNames.AddRange(
+            new string[]
+            {
+                "WebSockets",
+            }
+        );
+    }
+}
diff --git a/Engine/Plugins/Experimental/DataLink/Source/DataLinkWebSocket/Private/DataLinkWebSocket.cpp b/Engine/Plugins/Experimental/DataLink/Source/DataLinkWebSocket/Private/DataLinkWebSocket.cpp
new file mode 100644
--- /dev/null
+++ b/Engine/Plugins/Experimental/DataLink/Source/DataLinkWebSocket/Private/DataLinkWebSocket.cpp
@@ -0,0 +1,203 @@
+// Copyright Epic Games, Inc. All Rights Reserved.
+
+#include "DataLinkWebSocket.h"
+#include "DataLinkCoreTypes.h"
+#include "DataLinkExecutor.h"
+#include "DataLinkNames.h"
+#include "DataLinkNodeInstance.h"
+#include "DataLinkPinBuilder.h"
+#include "DataLinkWebSocketLog.h"
+#include "DataLinkWebSocketNames.h"
+#include "DataLinkWebSocketSettings.h"
+#include "DataLinkWebSocketSubsystem.h"
+#include "IWebSocket.h"
+
+#define LOCTEXT_NAMESPACE "DataLinkWebSocket"
+
+UDataLinkWebSocket::UDataLinkWebSocket()
+{
+	InstanceStruct = FDataLinkWebSocketInstanceData::StaticStruct();
+}
+
+void UDataLinkWebSocket::OnBuildPins(FDataLinkPinBuilder& Inputs, FDataLinkPinBuilder& Outputs) const
+{
+	Super::OnBuildPins(Inputs, Outputs);
+
+	Inputs.Add(UE::DataLink::InputWebSocketSettings)
+		.SetDisplayName(LOCTEXT("WebSocketSettingsDisplay", "Web Socket Settings"))
+		.SetStruct<FDataLinkWebSocketSettings>();
+
+	Inputs.Add(UE::DataLink::InputWebSocketMessages)
+		.SetDisplayName(LOCTEXT("WebSocketSettingsDisplay", "Connect Messages"))
+		.SetStruct<FDataLinkWebSocketMessages>();
+
+	Outputs.Add(UE::DataLink::OutputDefault)
+		.SetDisplayName(LOCTEXT("OutputDisplay", "Message Received"))
+		.SetStruct<FDataLinkString>();
+}
+
+EDataLinkExecutionReply UDataLinkWebSocket::OnExecute(FDataLinkExecutor& InExecutor) const
+{
+	UDataLinkWebSocketSubsystem* WebSocketSubsystem = UDataLinkWebSocketSubsystem::TryGet();
+	if (!WebSocketSubsystem)
+	{
+		return EDataLinkExecutionReply::Unhandled;
+	}
+
+	FDataLinkNodeInstance& NodeInstance = InExecutor.GetNodeInstanceMutable(this);
+	FDataLinkWebSocketInstanceData& InstanceData = NodeInstance.GetInstanceDataMutable().Get<FDataLinkWebSocketInstanceData>();
+
+	const FDataLinkInputDataViewer& InputDataViewer = NodeInstance.GetInputDataViewer();
+
+	const FDataLinkWebSocketSettings& WebSocketSettings = InputDataViewer.Get<FDataLinkWebSocketSettings>(UE::DataLink::InputWebSocketSettings);
+
+	TSharedPtr<IWebSocket> WebSocket = WebSocketSubsystem->FindWebSocket(InstanceData.WebSocketHandle);
+
+	if (WebSocket.IsValid() && WebSocket->IsConnected())
+	{
+		// If current web socket is valid and connected, early exit only if the web socket settings are the same.
+		if (InstanceData.WebSocketSettings.Equals(WebSocketSettings))
+		{
+			return EDataLinkExecutionReply::Handled;
+		}
+
+		UE_LOG(LogDataLinkWebSocket, Verbose, TEXT("[%*.s] Existing Web Socket found but connection settings differ from requested. Closing web socket.")
+			, InExecutor.GetContextName().Len(), InExecutor.GetContextName().GetData());
+	}
+
+	// Close web socket (existing is either non-existent, invalid, or connected but its connection settings are different)
+	if (InstanceData.WebSocketHandle.IsValid())
+	{
+		WebSocketSubsystem->CloseWebSocket(InstanceData.WebSocketHandle);
+		InstanceData.WebSocketHandle.Reset();
+	}
+
+	InstanceData.WebSocketSettings.Reset();
+
+	// Create Web Socket
+	{
+		UDataLinkWebSocketSubsystem::FCreateWebSocketResult CreateWebSocketResult;
+		if (!WebSocketSubsystem->CreateWebSocket(WebSocketSettings, CreateWebSocketResult))
+		{
+			UE_LOG(LogDataLinkWebSocket, Error, TEXT("[%*.s] Failed to create web socket"), InExecutor.GetContextName().Len(), InExecutor.GetContextName().GetData());
+			return EDataLinkExecutionReply::Unhandled;
+		}
+
+		WebSocket = MoveTemp(CreateWebSocketResult.WebSocket);
+		InstanceData.WebSocketHandle = MoveTemp(CreateWebSocketResult.Handle);
+	}
+
+	const TWeakPtr<FDataLinkExecutor> ExecutorWeak = InExecutor.AsWeak();
+
+	WebSocket->OnConnected().AddUObject(this, &UDataLinkWebSocket::OnConnected, ExecutorWeak);
+	WebSocket->OnConnectionError().AddUObject(this, &UDataLinkWebSocket::OnConnectionError, ExecutorWeak);
+	WebSocket->OnClosed().AddUObject(this, &UDataLinkWebSocket::OnClosed, ExecutorWeak);
+	WebSocket->OnMessage().AddUObject(this, &UDataLinkWebSocket::OnMessageReceived, ExecutorWeak);
+	WebSocket->Connect();
+
+	return EDataLinkExecutionReply::Handled;
+}
+
+void UDataLinkWebSocket::OnStop(const FDataLinkExecutor& InExecutor) const
+{
+	UDataLinkWebSocketSubsystem* WebSocketSubsystem = UDataLinkWebSocketSubsystem::TryGet();
+	if (!WebSocketSubsystem)
+	{
+		return;
+	}
+
+	const FDataLinkNodeInstance& NodeInstance = InExecutor.GetNodeInstance(this);
+
+	const FDataLinkWebSocketInstanceData& InstanceData = NodeInstance.GetInstanceData().Get<const FDataLinkWebSocketInstanceData>();
+
+	if (TSharedPtr<IWebSocket> WebSocket = WebSocketSubsystem->FindWebSocket(InstanceData.WebSocketHandle))
+	{
+		WebSocket->OnClosed().RemoveAll(this);
+	}
+
+	WebSocketSubsystem->CloseWebSocket(InstanceData.WebSocketHandle);
+}
+
+void UDataLinkWebSocket::OnConnected(TWeakPtr<FDataLinkExecutor> InExecutorWeak) const
+{
+	UDataLinkWebSocketSubsystem* WebSocketSubsystem = UDataLinkWebSocketSubsystem::TryGet();
+	if (!WebSocketSubsystem)
+	{
+		return;
+	}
+
+	TSharedPtr<FDataLinkExecutor> Executor = InExecutorWeak.Pin();
+	if (!Executor.IsValid())
+	{
+		return;
+	}
+
+	const FDataLinkNodeInstance& NodeInstance = Executor->GetNodeInstance(this);
+	const FDataLinkWebSocketInstanceData& InstanceData = NodeInstance.GetInstanceData().Get<const FDataLinkWebSocketInstanceData>();
+
+	TSharedPtr<IWebSocket> WebSocket = WebSocketSubsystem->FindWebSocket(InstanceData.WebSocketHandle);
+	if (!WebSocket.IsValid())
+	{
+		Executor->Fail(this);
+		return;
+	}
+
+	const FDataLinkInputDataViewer& InputDataViewer = NodeInstance.GetInputDataViewer();
+	const FDataLinkWebSocketMessages& WebSocketMessages = InputDataViewer.Get<FDataLinkWebSocketMessages>(UE::DataLink::InputWebSocketMessages);
+
+	for (const FString& Message : WebSocketMessages.ConnectMessages)
+	{
+		WebSocket->Send(Message);
+	}
+}
+
+void UDataLinkWebSocket::OnConnectionError(const FString& InError, TWeakPtr<FDataLinkExecutor> InExecutorWeak) const
+{
+	TSharedPtr<FDataLinkExecutor> Executor = InExecutorWeak.Pin();
+	if (!Executor.IsValid())
+	{
+		return;
+	}
+
+	UE_LOG(LogDataLinkWebSocket, Error, TEXT("[%*.s] Connection Error: %s")
+		, Executor->GetContextName().Len(), Executor->GetContextName().GetData()
+		, *InError);
+
+	Executor->Fail(this);
+}
+
+void UDataLinkWebSocket::OnClosed(int32 InStatusCode, const FString& InReason, bool bInWasClean, TWeakPtr<FDataLinkExecutor> InExecutorWeak) const
+{
+	TSharedPtr<FDataLinkExecutor> Executor = InExecutorWeak.Pin();
+	if (!Executor.IsValid())
+	{
+		return;
+	}
+
+	UE_LOG(LogDataLinkWebSocket, Error, TEXT("[%*.s] Web Socket Closed Status: %d --- %s")
+		, Executor->GetContextName().Len(), Executor->GetContextName().GetData()
+		, InStatusCode
+		, *InReason);
+
+	Executor->Fail(this);
+}
+
+void UDataLinkWebSocket::OnMessageReceived(const FString& InMessage, TWeakPtr<FDataLinkExecutor> InExecutorWeak) const
+{
+	TSharedPtr<FDataLinkExecutor> Executor = InExecutorWeak.Pin();
+	if (!Executor.IsValid())
+	{
+		return;
+	}
+
+	const FDataLinkNodeInstance& NodeInstance = Executor->GetNodeInstance(this);
+	const FDataLinkOutputDataViewer& OutputDataViewer = NodeInstance.GetOutputDataViewer();
+
+	FDataLinkString& OutputDataView = OutputDataViewer.Get<FDataLinkString>(UE::DataLink::OutputDefault);
+	OutputDataView.Value = InMessage;
+
+	// Keep execution running when message is received
+	Executor->NextPersist(this);
+}
+
+#undef LOCTEXT_NAMESPACE
diff --git a/Engine/Plugins/Experimental/DataLink/Source/DataLinkWebSocket/Private/DataLinkWebSocketHandle.cpp b/Engine/Plugins/Experimental/DataLink/Source/DataLinkWebSocket/Private/DataLinkWebSocketHandle.cpp
new file mode 100644
--- /dev/null
+++ b/Engine/Plugins/Experimental/DataLink/Source/DataLinkWebSocket/Private/DataLinkWebSocketHandle.cpp
@@ -0,0 +1,27 @@
+// Copyright Epic Games, Inc. All Rights Reserved.
+
+#include "DataLinkWebSocketHandle.h"
+
+namespace UE::DataLink
+{
+
+FWebSocketHandle FWebSocketHandle::GenerateNewHandle()
+{
+	static uint32 NextId = 0;
+
+	FWebSocketHandle Handle;
+
+	// Start from 1
+	Handle.Id = ++NextId;
+
+	// Handle edge case where NextId wraps around to 0.
+	// 0 is reserved for invalid handles.
+	if (Handle.Id == 0)
+	{
+		Handle.Id = ++NextId;
+	}
+
+	return Handle;
+}
+
+} // UE::DataLink
diff --git a/Engine/Plugins/Experimental/DataLink/Source/DataLinkWebSocket/Private/DataLinkWebSocketLog.h b/Engine/Plugins/Experimental/DataLink/Source/DataLinkWebSocket/Private/DataLinkWebSocketLog.h
new file mode 100644
--- /dev/null
+++ b/Engine/Plugins/Experimental/DataLink/Source/DataLinkWebSocket/Private/DataLinkWebSocketLog.h
@@ -0,0 +1,7 @@
+// Copyright Epic Games, Inc. All Rights Reserved.
+
+#pragma once
+
+#include "Logging/LogMacros.h"
+
+DECLARE_LOG_CATEGORY_EXTERN(LogDataLinkWebSocket, Log, All);
diff --git a/Engine/Plugins/Experimental/DataLink/Source/DataLinkWebSocket/Private/DataLinkWebSocketModule.cpp b/Engine/Plugins/Experimental/DataLink/Source/DataLinkWebSocket/Private/DataLinkWebSocketModule.cpp
new file mode 100644
--- /dev/null
+++ b/Engine/Plugins/Experimental/DataLink/Source/DataLinkWebSocket/Private/DataLinkWebSocketModule.cpp
@@ -0,0 +1,8 @@
+// Copyright Epic Games, Inc. All Rights Reserved.
+
+#include "DataLinkWebSocketLog.h"
+#include "Modules/ModuleManager.h"
+
+DEFINE_LOG_CATEGORY(LogDataLinkWebSocket);
+
+IMPLEMENT_MODULE(FDefaultModuleImpl, DataLinkWebSocket)
diff --git a/Engine/Plugins/Experimental/DataLink/Source/DataLinkWebSocket/Private/DataLinkWebSocketSettings.cpp b/Engine/Plugins/Experimental/DataLink/Source/DataLinkWebSocket/Private/DataLinkWebSocketSettings.cpp
new file mode 100644
--- /dev/null
+++ b/Engine/Plugins/Experimental/DataLink/Source/DataLinkWebSocket/Private/DataLinkWebSocketSettings.cpp
@@ -0,0 +1,44 @@
+// Copyright Epic Games, Inc. All Rights Reserved.
+
+#include "DataLinkWebSocketSettings.h"
+
+void FDataLinkWebSocketSettings::Reset()
+{
+	URL.Reset();
+	Protocols.Reset();
+	UpgradeHeaders.Reset();
+}
+
+bool FDataLinkWebSocketSettings::Equals(const FDataLinkWebSocketSettings& InOther) const
+{
+	constexpr ESearchCase::Type SearchCase = ESearchCase::CaseSensitive;
+
+	if (!URL.Equals(InOther.URL, SearchCase))
+	{
+		return false;
+	}
+
+	if (Protocols.Num() != InOther.Protocols.Num() || UpgradeHeaders.Num() != InOther.UpgradeHeaders.Num())
+	{
+		return false;
+	}
+
+	for (int32 Index = 0; Index < Protocols.Num(); ++Index)
+	{
+		if (!Protocols[Index].Equals(InOther.Protocols[Index], SearchCase))
+		{
+			return false;
+		}
+	}
+
+	for (const TPair<FString, FString>& HeaderPair : UpgradeHeaders)
+	{
+		const FString* HeaderValue = InOther.UpgradeHeaders.Find(HeaderPair.Key);
+		if (!HeaderValue || !HeaderPair.Value.Equals(*HeaderValue, SearchCase))
+		{
+			return false;
+		}
+	}
+
+	return true;
+}
diff --git a/Engine/Plugins/Experimental/DataLink/Source/DataLinkWebSocket/Private/DataLinkWebSocketSubsystem.cpp b/Engine/Plugins/Experimental/DataLink/Source/DataLinkWebSocket/Private/DataLinkWebSocketSubsystem.cpp
new file mode 100644
--- /dev/null
+++ b/Engine/Plugins/Experimental/DataLink/Source/DataLinkWebSocket/Private/DataLinkWebSocketSubsystem.cpp
@@ -0,0 +1,94 @@
+// Copyright Epic Games, Inc. All Rights Reserved.
+
+#include "DataLinkWebSocketSubsystem.h"
+#include "DataLinkWebSocketLog.h"
+#include "DataLinkWebSocketSettings.h"
+#include "Engine/Engine.h"
+#include "IWebSocket.h"
+#include "WebSocketsModule.h"
+
+UDataLinkWebSocketSubsystem* UDataLinkWebSocketSubsystem::TryGet()
+{
+	if (GEngine)
+	{
+		return GEngine->GetEngineSubsystem<UDataLinkWebSocketSubsystem>();
+	}
+	return nullptr;
+}
+
+bool UDataLinkWebSocketSubsystem::CreateWebSocket(const FDataLinkWebSocketSettings& InSettings, FCreateWebSocketResult& OutResult)
+{
+#if WITH_WEBSOCKETS
+	const UE::DataLink::FWebSocketHandle WebSocketHandle = UE::DataLink::FWebSocketHandle::GenerateNewHandle();
+
+	TSharedRef<IWebSocket> WebSocket = FWebSocketsModule::Get().CreateWebSocket(InSettings.URL, InSettings.Protocols, InSettings.UpgradeHeaders);
+	WebSockets.Add(WebSocketHandle, WebSocket);
+
+	OutResult.WebSocket = WebSocket;
+	OutResult.Handle = WebSocketHandle;
+	return true;
+#else
+	UE_LOG(LogDataLinkWebSocket, Error, TEXT("Failed to create web socket. Web Sockets are disabled."));
+	return false;
+#endif
+}
+
+void UDataLinkWebSocketSubsystem::CloseWebSocket(UE::DataLink::FWebSocketHandle InHandle)
+{
+	TSharedPtr<IWebSocket> WebSocket = FindWebSocket(InHandle);
+	if (!WebSocket.IsValid())
+	{
+		return;
+	}
+
+	WebSocket->OnClosed().AddWeakLambda(this,
+		[this, InHandle](int32 /*StatusCode*/, const FString& /*Reason*/, bool /*bWasClean*/)
+		{
+			check(IsInGameThread());
+
+			InvalidWebSockets.Add(WebSockets[InHandle]);
+			WebSockets.Remove(InHandle);
+
+			if (!CleanupHandle.IsValid())
+			{
+				constexpr float CleanupDelay = 1.f;
+
+				CleanupHandle = FTSTicker::GetCoreTicker().AddTicker(FTickerDelegate::CreateWeakLambda(this,
+					[This=this](float)
+					{
+						This->CleanupInvalidWebSockets();
+						return false; // one time call
+					})
+					, CleanupDelay);
+			}
+		});
+
+	WebSocket->Close();
+}
+
+TSharedPtr<IWebSocket> UDataLinkWebSocketSubsystem::FindWebSocket(UE::DataLink::FWebSocketHandle InHandle) const
+{
+	if (!InHandle.IsValid())
+	{
+		return nullptr;
+	}
+
+	if (const TSharedRef<IWebSocket>* WebSocket = WebSockets.Find(InHandle))
+	{
+		return *WebSocket;
+	}
+
+	return nullptr;
+}
+
+void UDataLinkWebSocketSubsystem::CleanupInvalidWebSockets()
+{
+	check(IsInGameThread());
+
+	// Removing CleanupHandle Ticker while still executing this ticker delegate is by-design valid.
+	// see TickerTests.cpp:112 "a delegate removal from inside the delegate execution (used to be a deadlock)"
+	FTSTicker::GetCoreTicker().RemoveTicker(CleanupHandle);
+	CleanupHandle.Reset();
+
+	InvalidWebSockets.Reset();
+}
diff --git a/Engine/Plugins/Experimental/DataLink/Source/DataLinkWebSocket/Private/DataLinkWebSocketSubsystem.h b/Engine/Plugins/Experimental/DataLink/Source/DataLinkWebSocket/Private/DataLinkWebSocketSubsystem.h
new file mode 100644
--- /dev/null
+++ b/Engine/Plugins/Experimental/DataLink/Source/DataLinkWebSocket/Private/DataLinkWebSocketSubsystem.h
@@ -0,0 +1,44 @@
+// Copyright Epic Games, Inc. All Rights Reserved.
+
+#pragma once
+
+#include "Containers/Ticker.h"
+#include "DataLinkWebSocketHandle.h"
+#include "Subsystems/EngineSubsystem.h"
+#include "DataLinkWebSocketSubsystem.generated.h"
+
+class IWebSocket;
+struct FDataLinkWebSocketSettings;
+
+UCLASS()
+class UDataLinkWebSocketSubsystem : public UEngineSubsystem
+{
+	GENERATED_BODY()
+
+public:
+	static UDataLinkWebSocketSubsystem* TryGet();
+
+	struct FCreateWebSocketResult
+	{
+		UE::DataLink::FWebSocketHandle Handle;
+		TSharedPtr<IWebSocket> WebSocket;
+	};
+	bool CreateWebSocket(const FDataLinkWebSocketSettings& InSettings, FCreateWebSocketResult& OutResult);
+
+	void CloseWebSocket(UE::DataLink::FWebSocketHandle InHandle);
+
+	/** Returns the web socket mapped to the handle if valid */
+	TSharedPtr<IWebSocket> FindWebSocket(UE::DataLink::FWebSocketHandle InHandle) const;
+
+private:
+	void CleanupInvalidWebSockets();
+
+	/** Registered web sockets */
+	TMap<UE::DataLink::FWebSocketHandle, TSharedRef<IWebSocket>> WebSockets;
+
+	/** Web sockets that have requested closure and are pending kill */
+	TArray<TSharedRef<IWebSocket>> InvalidWebSockets;
+
+	/** Handle to the callback to clean up invalid web sockets */
+	FTSTicker::FDelegateHandle CleanupHandle;
+};
diff --git a/Engine/Plugins/Experimental/DataLink/Source/DataLinkWebSocket/Public/DataLinkWebSocket.h b/Engine/Plugins/Experimental/DataLink/Source/DataLinkWebSocket/Public/DataLinkWebSocket.h
new file mode 100644
--- /dev/null
+++ b/Engine/Plugins/Experimental/DataLink/Source/DataLinkWebSocket/Public/DataLinkWebSocket.h
@@ -0,0 +1,68 @@
+// Copyright Epic Games, Inc. All Rights Reserved.
+
+#pragma once
+
+#include "DataLinkNode.h"
+#include "DataLinkWebSocketHandle.h"
+#include "DataLinkWebSocketSettings.h"
+#include "DataLinkWebSocket.generated.h"
+
+#define UE_API DATALINKWEBSOCKET_API
+
+/** The messages to send through web socket */
+USTRUCT(BlueprintType)
+struct FDataLinkWebSocketMessages
+{
+	GENERATED_BODY()
+
+	/** Messages to send */
+	UPROPERTY(EditAnywhere, BlueprintReadWrite, Category="Data Link Web Sockets")
+	TArray<FString> ConnectMessages;
+};
+
+/** Instance data for the Web Socket Data Link node */
+USTRUCT()
+struct FDataLinkWebSocketInstanceData
+{
+	GENERATED_BODY()
+
+	/** Handle to the last web socket created for this instance */
+	UE::DataLink::FWebSocketHandle WebSocketHandle;
+
+	/**
+	 * Settings used to create the web socket
+	 * Used to determine if a new web socket has to be created
+	 */
+	UPROPERTY()
+	FDataLinkWebSocketSettings WebSocketSettings;
+};
+
+UCLASS(MinimalAPI, DisplayName="Web Socket", Category="Web Socket")
+class UDataLinkWebSocket : public UDataLinkNode
+{
+	GENERATED_BODY()
+
+	UDataLinkWebSocket();
+
+protected:
+	//~ Begin UDataLinkNode
+	UE_API virtual void OnBuildPins(FDataLinkPinBuilder& Inputs, FDataLinkPinBuilder& Outputs) const override;
+	UE_API virtual EDataLinkExecutionReply OnExecute(FDataLinkExecutor& InExecutor) const override;
+	UE_API virtual void OnStop(const FDataLinkExecutor& InExecutor) const override;
+	//~ End UDataLinkNode
+
+private:
+	/** Called when the web socket has connected */
+	void OnConnected(TWeakPtr<FDataLinkExecutor> InExecutorWeak) const;
+
+	/** Called when the web socket encountered a connection error */
+	void OnConnectionError(const FString& InError, TWeakPtr<FDataLinkExecutor> InExecutorWeak) const;
+
+	/** Called when the web socket has closed */
+	void OnClosed(int32 InStatusCode, const FString& InReason, bool bInWasClean, TWeakPtr<FDataLinkExecutor> InExecutorWeak) const;
+
+	/** Called when the web socket has a received a message */
+	void OnMessageReceived(const FString& InMessage, TWeakPtr<FDataLinkExecutor> InExecutorWeak) const;
+};
+
+#undef UE_API
diff --git a/Engine/Plugins/Experimental/DataLink/Source/DataLinkWebSocket/Public/DataLinkWebSocketHandle.h b/Engine/Plugins/Experimental/DataLink/Source/DataLinkWebSocket/Public/DataLinkWebSocketHandle.h
new file mode 100644
--- /dev/null
+++ b/Engine/Plugins/Experimental/DataLink/Source/DataLinkWebSocket/Public/DataLinkWebSocketHandle.h
@@ -0,0 +1,40 @@
+// Copyright Epic Games, Inc. All Rights Reserved.
+
+#pragma once
+
+#include "HAL/Platform.h"
+#include "Templates/TypeHash.h"
+
+namespace UE::DataLink
+{
+
+class FWebSocketHandle
+{
+public:
+	static FWebSocketHandle GenerateNewHandle();
+
+	void Reset()
+	{
+		Id = 0;
+	}
+
+	bool IsValid() const
+	{
+		return Id != 0;
+	}
+
+	bool operator==(const FWebSocketHandle& InOther) const
+	{
+		return Id == InOther.Id;
+	}
+
+	friend uint32 GetTypeHash(const FWebSocketHandle& InHandle)
+	{
+		return ::GetTypeHash(InHandle.Id);
+	}
+
+private:
+	uint32 Id = 0;
+};
+
+} // UE::DataLink
diff --git a/Engine/Plugins/Experimental/DataLink/Source/DataLinkWebSocket/Public/DataLinkWebSocketNames.h b/Engine/Plugins/Experimental/DataLink/Source/DataLinkWebSocket/Public/DataLinkWebSocketNames.h
new file mode 100644
--- /dev/null
+++ b/Engine/Plugins/Experimental/DataLink/Source/DataLinkWebSocket/Public/DataLinkWebSocketNames.h
@@ -0,0 +1,11 @@
+// Copyright Epic Games, Inc. All Rights Reserved.
+
+#pragma once
+
+#include "UObject/NameTypes.h"
+
+namespace UE::DataLink
+{
+	const FLazyName InputWebSocketSettings(TEXT("InputWebSocketSettings"));
+	const FLazyName InputWebSocketMessages(TEXT("InputWebSocketMessages"));
+}
diff --git a/Engine/Plugins/Experimental/DataLink/Source/DataLinkWebSocket/Public/DataLinkWebSocketSettings.h b/Engine/Plugins/Experimental/DataLink/Source/DataLinkWebSocket/Public/DataLinkWebSocketSettings.h
new file mode 100644
--- /dev/null
+++ b/Engine/Plugins/Experimental/DataLink/Source/DataLinkWebSocket/Public/DataLinkWebSocketSettings.h
@@ -0,0 +1,36 @@
+// Copyright Epic Games, Inc. All Rights Reserved.
+
+#pragma once
+
+#include "Containers/Array.h"
+#include "Containers/Map.h"
+#include "Containers/UnrealString.h"
+#include "DataLinkWebSocketSettings.generated.h"
+
+#define UE_API DATALINKWEBSOCKET_API
+
+USTRUCT(BlueprintType)
+struct FDataLinkWebSocketSettings
+{
+	GENERATED_BODY()
+
+	/** Resets all the settings */
+	UE_API void Reset();
+
+	/** Returns true if the settings matches the given other settings */
+	UE_API bool Equals(const FDataLinkWebSocketSettings& InOther) const;
+
+	/** URL to connect to where the web socket server will respond */
+	UPROPERTY(EditAnywhere, BlueprintReadWrite, Category="Data Link Web Sockets")
+	FString URL;
+
+	/** List of the protocols to use with the connection */
+	UPROPERTY(EditAnywhere, BlueprintReadWrite, Category="Data Link Web Sockets")
+	TArray<FString> Protocols;
+
+	/** The upgrade headers to send with the upgrade request */
+	UPROPERTY(EditAnywhere, BlueprintReadWrite, Category="Data Link Web Sockets")
+	TMap<FString, FString> UpgradeHeaders;
+};
+
+#undef UE_API
