Motion Design Data Link: Fix web socket repeated loctext key

diff --git a/Engine/Plugins/Experimental/DataLink/Source/DataLinkWebSocket/Private/DataLinkWebSocket.cpp b/Engine/Plugins/Experimental/DataLink/Source/DataLinkWebSocket/Private/DataLinkWebSocket.cpp
--- a/Engine/Plugins/Experimental/DataLink/Source/DataLinkWebSocket/Private/DataLinkWebSocket.cpp
+++ b/Engine/Plugins/Experimental/DataLink/Source/DataLinkWebSocket/Private/DataLinkWebSocket.cpp
@@ -28,7 +28,7 @@ void UDataLinkWebSocket::OnBuildPins(FDataLinkPinBuilder& Inputs, FDataLinkPinBu
 		.SetStruct<FDataLinkWebSocketSettings>();
 
 	Inputs.Add(UE::DataLink::InputWebSocketMessages)
-		.SetDisplayName(LOCTEXT("WebSocketSettingsDisplay", "Connect Messages"))
+		.SetDisplayName(LOCTEXT("WebSocketMessagesDisplay", "Connect Messages"))
 		.SetStruct<FDataLinkWebSocketMessages>();
 
 	Outputs.Add(UE::DataLink::OutputDefault)
