Motion Design Data Link: fix issue with websocket logging for stringviews + added stop function call for the motion design data link actor

diff --git a/Engine/Plugins/Experimental/AvalancheDataLink/Source/AvalancheDataLink/Private/AvaDataLinkActor.cpp b/Engine/Plugins/Experimental/AvalancheDataLink/Source/AvalancheDataLink/Private/AvaDataLinkActor.cpp
--- a/Engine/Plugins/Experimental/AvalancheDataLink/Source/AvalancheDataLink/Private/AvaDataLinkActor.cpp
+++ b/Engine/Plugins/Experimental/AvalancheDataLink/Source/AvalancheDataLink/Private/AvaDataLinkActor.cpp
@@ -13,3 +13,14 @@ void AAvaDataLinkActor::ExecuteDataLinkInstances()
 		}
 	}
 }
+
+void AAvaDataLinkActor::StopDataLinkInstances()
+{
+	for (UAvaDataLinkInstance* DataLinkInstance: DataLinkInstances)
+	{
+		if (DataLinkInstance)
+		{
+			DataLinkInstance->Stop();
+		}
+	}
+}
diff --git a/Engine/Plugins/Experimental/AvalancheDataLink/Source/AvalancheDataLink/Private/AvaDataLinkActor.h b/Engine/Plugins/Experimental/AvalancheDataLink/Source/AvalancheDataLink/Private/AvaDataLinkActor.h
--- a/Engine/Plugins/Experimental/AvalancheDataLink/Source/AvalancheDataLink/Private/AvaDataLinkActor.h
+++ b/Engine/Plugins/Experimental/AvalancheDataLink/Source/AvalancheDataLink/Private/AvaDataLinkActor.h
@@ -16,6 +16,9 @@ public:
 	UFUNCTION(BlueprintCallable, CallInEditor, Category="Motion Design Data Link")
 	void ExecuteDataLinkInstances();
 
+	UFUNCTION(BlueprintCallable, CallInEditor, Category="Motion Design Data Link")
+	void StopDataLinkInstances();
+
 	UPROPERTY(EditAnywhere, Instanced, Category="Motion Design Data Link")
 	TArray<TObjectPtr<UAvaDataLinkInstance>> DataLinkInstances;
 };
diff --git a/Engine/Plugins/Experimental/DataLink/Source/DataLinkWebSocket/Private/DataLinkWebSocket.cpp b/Engine/Plugins/Experimental/DataLink/Source/DataLinkWebSocket/Private/DataLinkWebSocket.cpp
--- a/Engine/Plugins/Experimental/DataLink/Source/DataLinkWebSocket/Private/DataLinkWebSocket.cpp
+++ b/Engine/Plugins/Experimental/DataLink/Source/DataLinkWebSocket/Private/DataLinkWebSocket.cpp
@@ -61,7 +61,7 @@ EDataLinkExecutionReply UDataLinkWebSocket::OnExecute(FDataLinkExecutor& InExecu
 			return EDataLinkExecutionReply::Handled;
 		}
 
-		UE_LOG(LogDataLinkWebSocket, Verbose, TEXT("[%*.s] Existing Web Socket found but connection settings differ from requested. Closing web socket.")
+		UE_LOG(LogDataLinkWebSocket, Verbose, TEXT("[%.*s] Existing Web Socket found but connection settings differ from requested. Closing web socket.")
 			, InExecutor.GetContextName().Len(), InExecutor.GetContextName().GetData());
 	}
 
@@ -79,7 +79,7 @@ EDataLinkExecutionReply UDataLinkWebSocket::OnExecute(FDataLinkExecutor& InExecu
 		UDataLinkWebSocketSubsystem::FCreateWebSocketResult CreateWebSocketResult;
 		if (!WebSocketSubsystem->CreateWebSocket(WebSocketSettings, CreateWebSocketResult))
 		{
-			UE_LOG(LogDataLinkWebSocket, Error, TEXT("[%*.s] Failed to create web socket"), InExecutor.GetContextName().Len(), InExecutor.GetContextName().GetData());
+			UE_LOG(LogDataLinkWebSocket, Error, TEXT("[%.*s] Failed to create web socket"), InExecutor.GetContextName().Len(), InExecutor.GetContextName().GetData());
 			return EDataLinkExecutionReply::Unhandled;
 		}
 
@@ -159,7 +159,7 @@ void UDataLinkWebSocket::OnConnectionError(const FString& InError, TWeakPtr<FDat
 		return;
 	}
 
-	UE_LOG(LogDataLinkWebSocket, Error, TEXT("[%*.s] Connection Error: %s")
+	UE_LOG(LogDataLinkWebSocket, Error, TEXT("[%.*s] Connection Error: %s")
 		, Executor->GetContextName().Len(), Executor->GetContextName().GetData()
 		, *InError);
 
@@ -174,7 +174,7 @@ void UDataLinkWebSocket::OnClosed(int32 InStatusCode, const FString& InReason, b
 		return;
 	}
 
-	UE_LOG(LogDataLinkWebSocket, Error, TEXT("[%*.s] Web Socket Closed Status: %d --- %s")
+	UE_LOG(LogDataLinkWebSocket, Error, TEXT("[%.*s] Web Socket Closed Status: %d --- %s")
 		, Executor->GetContextName().Len(), Executor->GetContextName().GetData()
 		, InStatusCode
 		, *InReason);
