Motion Design Data Link: deprecated request data link task and opted now to use a data link object for blueprint interfacing

diff --git a/Engine/Plugins/Experimental/DataLink/Source/DataLink/Private/DataLinkExecutorArguments.cpp b/Engine/Plugins/Experimental/DataLink/Source/DataLink/Private/DataLinkExecutorArguments.cpp
--- a/Engine/Plugins/Experimental/DataLink/Source/DataLink/Private/DataLinkExecutorArguments.cpp
+++ b/Engine/Plugins/Experimental/DataLink/Source/DataLink/Private/DataLinkExecutorArguments.cpp
@@ -14,6 +14,12 @@ FDataLinkExecutorArguments::FDataLinkExecutorArguments(FDataLinkInstance&& InIns
 }
 
 #if WITH_DATALINK_CONTEXT
+FDataLinkExecutorArguments&& FDataLinkExecutorArguments::SetContextName(const FString& InContextName)
+{
+	ContextName = InContextName;
+	return MoveTemp(*this);
+}
+
 FDataLinkExecutorArguments&& FDataLinkExecutorArguments::SetContextName(FString&& InContextName)
 {
 	ContextName = MoveTemp(InContextName);
diff --git a/Engine/Plugins/Experimental/DataLink/Source/DataLink/Private/DataLinkExecutorObject.cpp b/Engine/Plugins/Experimental/DataLink/Source/DataLink/Private/DataLinkExecutorObject.cpp
new file mode 100644
--- /dev/null
+++ b/Engine/Plugins/Experimental/DataLink/Source/DataLink/Private/DataLinkExecutorObject.cpp
@@ -0,0 +1,46 @@
+// Copyright Epic Games, Inc. All Rights Reserved.
+
+#include "DataLinkExecutorObject.h"
+#include "DataLinkExecutor.h"
+#include "DataLinkExecutorArguments.h"
+#include "DataLinkUtils.h"
+
+bool UDataLinkExecutorObject::IsRunning() const
+{
+	return Executor.IsValid() && Executor->IsRunning();
+}
+
+void UDataLinkExecutorObject::Run()
+{
+	Stop();
+
+	Executor = FDataLinkExecutor::Create(FDataLinkExecutorArguments(DataLinkInstance)
+#if WITH_DATALINK_CONTEXT
+		.SetContextName(ContextName)
+#endif
+		.SetContextObject(GetOuter())
+		.SetSink(UE::DataLink::TryGetSink(SinkProvider))
+		.SetOnOutputData(FOnDataLinkOutputData::CreateUObject(this, &UDataLinkExecutorObject::NotifyOutputDataReceived))
+		.SetOnFinished(FOnDataLinkExecutionFinished::CreateUObject(this, &UDataLinkExecutorObject::NotifyExecutionFinished)));
+
+	Executor->Run();
+}
+
+void UDataLinkExecutorObject::Stop()
+{
+	if (Executor.IsValid())
+	{
+		Executor->Stop();
+		Executor.Reset();
+	}
+}
+
+void UDataLinkExecutorObject::NotifyOutputDataReceived(const FDataLinkExecutor& InExecutor, FConstStructView InOutputDataView)
+{
+	OnOutputData.Broadcast(FInstancedStruct(InOutputDataView));
+}
+
+void UDataLinkExecutorObject::NotifyExecutionFinished(const FDataLinkExecutor& InExecutor, EDataLinkExecutionResult InExecutionResult)
+{
+	OnExecutionFinished.Broadcast(InExecutionResult);
+}
diff --git a/Engine/Plugins/Experimental/DataLink/Source/DataLink/Private/DataLinkRequestProxy.cpp b/Engine/Plugins/Experimental/DataLink/Source/DataLink/Private/DataLinkRequestProxy.cpp
--- a/Engine/Plugins/Experimental/DataLink/Source/DataLink/Private/DataLinkRequestProxy.cpp
+++ b/Engine/Plugins/Experimental/DataLink/Source/DataLink/Private/DataLinkRequestProxy.cpp
@@ -7,6 +7,7 @@
 #include "DataLinkExecutorArguments.h"
 #include "DataLinkUtils.h"
 
+PRAGMA_DISABLE_DEPRECATION_WARNINGS
 UDataLinkRequestProxy* UDataLinkRequestProxy::CreateRequestProxy(FDataLinkInstance InDataLinkInstance, UObject* InExecutionContext, TScriptInterface<IDataLinkSinkProvider> InDataLinkSinkProvider)
 {
 	UDataLinkRequestProxy* RequestProxy = NewObject<UDataLinkRequestProxy>();
@@ -14,6 +15,7 @@ UDataLinkRequestProxy* UDataLinkRequestProxy::CreateRequestProxy(FDataLinkInstan
 	RequestProxy->ProcessRequest(MoveTemp(InDataLinkInstance), InExecutionContext, InDataLinkSinkProvider);
 	return RequestProxy;
 }
+PRAGMA_ENABLE_DEPRECATION_WARNINGS
 
 void UDataLinkRequestProxy::ProcessRequest(FDataLinkInstance&& InDataLinkInstance, UObject* InExecutionContext, TScriptInterface<IDataLinkSinkProvider> InDataLinkSinkProvider)
 {
@@ -23,7 +25,9 @@ void UDataLinkRequestProxy::ProcessRequest(FDataLinkInstance&& InDataLinkInstanc
 #endif
 		.SetContextObject(InExecutionContext)
 		.SetSink(UE::DataLink::TryGetSink(InDataLinkSinkProvider))
+PRAGMA_DISABLE_DEPRECATION_WARNINGS
 		.SetOnOutputData(FOnDataLinkOutputData::CreateUObject(this, &UDataLinkRequestProxy::OnOutputDataReceived)));
+PRAGMA_ENABLE_DEPRECATION_WARNINGS
 
 	// Delay execution by 1 frame. This is so that even for executions that finish immediately,
 	// the Async Call blueprint script still has chance to get the returned Proxy and listen to execution finished event
diff --git a/Engine/Plugins/Experimental/DataLink/Source/DataLink/Public/DataLinkExecutorArguments.h b/Engine/Plugins/Experimental/DataLink/Source/DataLink/Public/DataLinkExecutorArguments.h
--- a/Engine/Plugins/Experimental/DataLink/Source/DataLink/Public/DataLinkExecutorArguments.h
+++ b/Engine/Plugins/Experimental/DataLink/Source/DataLink/Public/DataLinkExecutorArguments.h
@@ -20,6 +20,8 @@ public:
 	DATALINK_API FDataLinkExecutorArguments(FDataLinkInstance&& InInstance);
 
 #if WITH_DATALINK_CONTEXT
+	/** Added context information of the execution */
+	DATALINK_API FDataLinkExecutorArguments&& SetContextName(const FString& InContextName);
 	/** Added context information of the execution */
 	DATALINK_API FDataLinkExecutorArguments&& SetContextName(FString&& InContextName);
 #endif
diff --git a/Engine/Plugins/Experimental/DataLink/Source/DataLink/Public/DataLinkExecutorObject.h b/Engine/Plugins/Experimental/DataLink/Source/DataLink/Public/DataLinkExecutorObject.h
new file mode 100644
--- /dev/null
+++ b/Engine/Plugins/Experimental/DataLink/Source/DataLink/Public/DataLinkExecutorObject.h
@@ -0,0 +1,66 @@
+// Copyright Epic Games, Inc. All Rights Reserved.
+
+#pragma once
+
+#include "DataLinkDelegates.h"
+#include "DataLinkInstance.h"
+#include "UObject/Object.h"
+#include "DataLinkExecutorObject.generated.h"
+
+#define UE_API DATALINK_API
+
+class IDataLinkSinkProvider;
+
+DECLARE_DYNAMIC_MULTICAST_DELEGATE_OneParam(FReceiveDataLinkOutputData, FInstancedStruct, OutputData);
+DECLARE_DYNAMIC_MULTICAST_DELEGATE_OneParam(FReceiveDataLinkExecutionFinished, EDataLinkExecutionResult, Result);
+
+/** UObject wrapper for Data Link Executor for Blueprint usage */
+UCLASS(MinimalAPI, Blueprintable, BlueprintType)
+class UDataLinkExecutorObject : public UObject
+{
+	GENERATED_BODY()
+
+	/** Returns whether there's an active data link execution happening */
+	UFUNCTION(BlueprintPure, Category="Data Link")
+	UE_API bool IsRunning() const;
+
+	/** Executes the Data Link Instance, stopping the existing active execution if any */
+	UFUNCTION(BlueprintCallable, Category="Data Link")
+	UE_API void Run();
+
+	/** Stops an existing data link execution if active */
+	UFUNCTION(BlueprintCallable, Category="Data Link")
+	UE_API void Stop();
+
+protected:
+	/** Called when data has been received. Calls the OnOutputData delegate */
+	void NotifyOutputDataReceived(const FDataLinkExecutor& InExecutor, FConstStructView InOutputDataView);
+
+	/** Called when execution has finished. Calls the OnExecutionFinished delegate */
+	void NotifyExecutionFinished(const FDataLinkExecutor& InExecutor, EDataLinkExecutionResult InExecutionResult);
+
+	/** Data link graph and input data to execute */
+	UPROPERTY(EditAnywhere, BlueprintReadWrite, Category="Data Link", meta=(ExposeOnSpawn="true"))
+	FDataLinkInstance DataLinkInstance;
+
+	/** Additional context to identify the data link execution */
+	UPROPERTY(EditAnywhere, BlueprintReadWrite, Category="Data Link", meta=(ExposeOnSpawn="true"), AdvancedDisplay)
+	FString ContextName;
+
+	/** Optional: Sink to use for cross execution storage */
+	UPROPERTY(EditAnywhere, BlueprintReadWrite, Category="Data Link", meta=(ExposeOnSpawn="true"), AdvancedDisplay)
+	TScriptInterface<IDataLinkSinkProvider> SinkProvider;
+
+	/** Delegate called when data link execution has output data ready */
+	UPROPERTY(BlueprintAssignable, Category="Data Link", meta=(HideInDetailPanel))
+	FReceiveDataLinkOutputData OnOutputData;
+
+	/** Delegate called when data link execution has finished completely */
+	UPROPERTY(BlueprintAssignable, Category="Data Link", meta=(HideInDetailPanel))
+	FReceiveDataLinkExecutionFinished OnExecutionFinished;
+
+	/** The underlying executor that this object wraps */
+	TSharedPtr<FDataLinkExecutor> Executor;
+};
+
+#undef UE_API
diff --git a/Engine/Plugins/Experimental/DataLink/Source/DataLink/Public/DataLinkRequestProxy.h b/Engine/Plugins/Experimental/DataLink/Source/DataLink/Public/DataLinkRequestProxy.h
--- a/Engine/Plugins/Experimental/DataLink/Source/DataLink/Public/DataLinkRequestProxy.h
+++ b/Engine/Plugins/Experimental/DataLink/Source/DataLink/Public/DataLinkRequestProxy.h
@@ -14,7 +14,7 @@ struct FDataLinkInstance;
 DECLARE_DYNAMIC_MULTICAST_DELEGATE_TwoParams(FOnDataLinkOutputDataProxy, const FInstancedStruct&, OutputData, EDataLinkExecutionResult, ExecutionResult);
 
 UCLASS(MinimalAPI)
-class UDataLinkRequestProxy : public UObject
+class UE_DEPRECATED(5.7, "Data Link Request Proxy deprecated. Use Data Link Executor Object instead") UDataLinkRequestProxy : public UObject
 {
 	GENERATED_BODY()
 
@@ -22,10 +22,13 @@ public:
 	UPROPERTY(BlueprintAssignable)
 	FOnDataLinkOutputDataProxy OnOutputData;
 
+PRAGMA_DISABLE_DEPRECATION_WARNINGS
+	UE_DEPRECATED(5.7, "Data Link Request Proxy deprecated. Use Data Link Executor Object instead")
 	UFUNCTION(BlueprintCallable, Category="Data Link", meta=(BlueprintInternalUseOnly))
 	static DATALINK_API UDataLinkRequestProxy* CreateRequestProxy(FDataLinkInstance InDataLinkInstance
 		, UObject* InExecutionContext
 		, TScriptInterface<IDataLinkSinkProvider> InDataLinkSinkProvider);
+PRAGMA_ENABLE_DEPRECATION_WARNINGS
 
 	DATALINK_API void ProcessRequest(FDataLinkInstance&& InDataLinkInstance
 		, UObject* InExecutionContext
diff --git a/Engine/Plugins/Experimental/DataLink/Source/DataLinkEdGraph/Private/Blueprint/K2Node_AsyncDataLinkRequest.cpp b/Engine/Plugins/Experimental/DataLink/Source/DataLinkEdGraph/Private/Blueprint/K2Node_AsyncDataLinkRequest.cpp
--- a/Engine/Plugins/Experimental/DataLink/Source/DataLinkEdGraph/Private/Blueprint/K2Node_AsyncDataLinkRequest.cpp
+++ b/Engine/Plugins/Experimental/DataLink/Source/DataLinkEdGraph/Private/Blueprint/K2Node_AsyncDataLinkRequest.cpp
@@ -5,9 +5,11 @@
 
 UK2Node_AsyncDataLinkRequest::UK2Node_AsyncDataLinkRequest()
 {
+PRAGMA_DISABLE_DEPRECATION_WARNINGS
 	ProxyFactoryFunctionName = GET_FUNCTION_NAME_CHECKED(UDataLinkRequestProxy, CreateRequestProxy);
 	ProxyFactoryClass = UDataLinkRequestProxy::StaticClass();
 	ProxyClass = UDataLinkRequestProxy::StaticClass();
+PRAGMA_ENABLE_DEPRECATION_WARNINGS
 }
 
 UEdGraphPin* UK2Node_AsyncDataLinkRequest::GetDataLinkInstancePin() const
diff --git a/Engine/Plugins/Experimental/DataLink/Source/DataLinkEdGraph/Private/Blueprint/K2Node_DataLinkRequest.cpp b/Engine/Plugins/Experimental/DataLink/Source/DataLinkEdGraph/Private/Blueprint/K2Node_DataLinkRequest.cpp
--- a/Engine/Plugins/Experimental/DataLink/Source/DataLinkEdGraph/Private/Blueprint/K2Node_DataLinkRequest.cpp
+++ b/Engine/Plugins/Experimental/DataLink/Source/DataLinkEdGraph/Private/Blueprint/K2Node_DataLinkRequest.cpp
@@ -21,26 +21,6 @@ const FLazyName UK2Node_DataLinkRequest::PN_DataLinkSinkProvider = TEXT("DataLin
 const FLazyName UK2Node_DataLinkRequest::PN_Success = TEXT("Success");
 const FLazyName UK2Node_DataLinkRequest::PN_Failure = TEXT("Failure");
 
-UEdGraphPin* UK2Node_DataLinkRequest::GetDataLinkInstancePin() const
-{
-	return FindPinChecked(UK2Node_DataLinkRequest::PN_DataLinkInstance, EGPD_Input);
-}
-
-UEdGraphPin* UK2Node_DataLinkRequest::GetExecutionContextPin() const
-{
-	return FindPinChecked(UK2Node_DataLinkRequest::PN_ExecutionContext, EGPD_Input);
-}
-
-UEdGraphPin* UK2Node_DataLinkRequest::GetDataLinkSinkProviderPin() const
-{
-	return FindPinChecked(UK2Node_DataLinkRequest::PN_DataLinkSinkProvider, EGPD_Input);
-}
-
-UEdGraphPin* UK2Node_DataLinkRequest::GetOutputPin() const
-{
-	return FindPinChecked(UDataLinkEdNode::PN_Output, EGPD_Output);
-}
-
 void UK2Node_DataLinkRequest::AllocateDefaultPins()
 {
 	Super::AllocateDefaultPins();
@@ -53,7 +33,7 @@ void UK2Node_DataLinkRequest::AllocateDefaultPins()
 	ThenPin->PinToolTip = TEXT("For execution that needs to happen immediately without waiting for a response.");
 
 	// Success Pin
-	UEdGraphPin* const SuccessPin = CreatePin(EGPD_Output, UEdGraphSchema_K2::PC_Exec, UK2Node_DataLinkRequest::PN_Success);
+	UEdGraphPin* const SuccessPin = CreatePin(EGPD_Output, UEdGraphSchema_K2::PC_Exec, PN_Success);
 	SuccessPin->PinFriendlyName = LOCTEXT("SuccessPinLabel", "Output Data");
 	SuccessPin->PinToolTip = TEXT("Called when the execution received and processed data successfully");
 
@@ -61,15 +41,15 @@ void UK2Node_DataLinkRequest::AllocateDefaultPins()
 	CreatePin(EGPD_Output, UEdGraphSchema_K2::PC_Struct, FInstancedStruct::StaticStruct(), UDataLinkEdNode::PN_Output);
 
 	// Failure pin
-	UEdGraphPin* const FailurePin = CreatePin(EGPD_Output, UEdGraphSchema_K2::PC_Exec, UK2Node_DataLinkRequest::PN_Failure);
+	UEdGraphPin* const FailurePin = CreatePin(EGPD_Output, UEdGraphSchema_K2::PC_Exec, PN_Failure);
 	FailurePin->PinFriendlyName = LOCTEXT("FailurePinLabel", "Data Link Failed");
 	FailurePin->PinToolTip = TEXT("Called when the execution failed and no data will be received");
 
 	// Create Data Link Instance Pin
-	CreatePin(EGPD_Input, UEdGraphSchema_K2::PC_Struct, FDataLinkInstance::StaticStruct(), UK2Node_DataLinkRequest::PN_DataLinkInstance);
+	CreatePin(EGPD_Input, UEdGraphSchema_K2::PC_Struct, FDataLinkInstance::StaticStruct(), PN_DataLinkInstance);
 
 	// Create Execution Context Pin
-	if (UEdGraphPin* ExecutionContextPin = CreatePin(EGPD_Input, UEdGraphSchema_K2::PC_Object, UObject::StaticClass(), UK2Node_DataLinkRequest::PN_ExecutionContext))
+	if (UEdGraphPin* ExecutionContextPin = CreatePin(EGPD_Input, UEdGraphSchema_K2::PC_Object, UObject::StaticClass(), PN_ExecutionContext))
 	{
 		// Prevent user from editing default value (i.e. hiding the object dropdown as option).
 		ExecutionContextPin->bDefaultValueIsIgnored = true;
@@ -77,7 +57,7 @@ void UK2Node_DataLinkRequest::AllocateDefaultPins()
 	}
 
 	// Create Data Link Sink Provider Pin
-	if (UEdGraphPin* SinkProviderPin = CreatePin(EGPD_Input, UEdGraphSchema_K2::PC_Interface, UDataLinkSinkProvider::StaticClass(), UK2Node_DataLinkRequest::PN_DataLinkSinkProvider))
+	if (UEdGraphPin* SinkProviderPin = CreatePin(EGPD_Input, UEdGraphSchema_K2::PC_Interface, UDataLinkSinkProvider::StaticClass(), PN_DataLinkSinkProvider))
 	{
 		// Prevent user from editing default value (i.e. hiding the object dropdown as option).
 		SinkProviderPin->bDefaultValueIsIgnored = true;
@@ -107,6 +87,23 @@ FText UK2Node_DataLinkRequest::GetNodeTitle(ENodeTitleType::Type InTitleType) co
 	return LOCTEXT("NodeTitle", "Execute Data Link");
 }
 
+bool UK2Node_DataLinkRequest::IsDeprecated() const
+{
+	return true;
+}
+
+FEdGraphNodeDeprecationResponse UK2Node_DataLinkRequest::GetDeprecationResponse(EEdGraphNodeDeprecationType InDeprecationType) const
+{
+	if (InDeprecationType == EEdGraphNodeDeprecationType::NodeTypeIsDeprecated)
+	{
+		FEdGraphNodeDeprecationResponse Response;
+		Response.MessageType = EEdGraphNodeDeprecationMessageType::Warning;
+		Response.MessageText = LOCTEXT("NodeDeprecated_Warning", "@@ is deprecated. Prefer creating a Data Link Executor Object and using that instead.");
+		return Response;
+	}
+	return Super::GetDeprecationResponse(InDeprecationType);
+}
+
 void UK2Node_DataLinkRequest::ExpandNode(FKismetCompilerContext& InCompilerContext, UEdGraph* InSourceGraph)
 {
 	UK2Node_AsyncDataLinkRequest* AsyncRequestNode = InCompilerContext.SpawnIntermediateNode<UK2Node_AsyncDataLinkRequest>(this, InSourceGraph);
@@ -114,10 +111,10 @@ void UK2Node_DataLinkRequest::ExpandNode(FKismetCompilerContext& InCompilerConte
 
 	InCompilerContext.MovePinLinksToIntermediate(*GetExecPin(), *AsyncRequestNode->GetExecPin());
 	InCompilerContext.MovePinLinksToIntermediate(*GetThenPin(), *AsyncRequestNode->GetThenPin());
-	InCompilerContext.MovePinLinksToIntermediate(*GetDataLinkInstancePin(), *AsyncRequestNode->GetDataLinkInstancePin());
-	InCompilerContext.MovePinLinksToIntermediate(*GetExecutionContextPin(), *AsyncRequestNode->GetExecutionContextPin());
-	InCompilerContext.MovePinLinksToIntermediate(*GetDataLinkSinkProviderPin(), *AsyncRequestNode->GetDataLinkSinkProviderPin());
-	InCompilerContext.MovePinLinksToIntermediate(*GetOutputPin(), *AsyncRequestNode->GetOutputDataPin());
+	InCompilerContext.MovePinLinksToIntermediate(*FindPinChecked(PN_DataLinkInstance, EGPD_Input), *AsyncRequestNode->GetDataLinkInstancePin());
+	InCompilerContext.MovePinLinksToIntermediate(*FindPinChecked(PN_ExecutionContext, EGPD_Input), *AsyncRequestNode->GetExecutionContextPin());
+	InCompilerContext.MovePinLinksToIntermediate(*FindPinChecked(PN_DataLinkSinkProvider, EGPD_Input), *AsyncRequestNode->GetDataLinkSinkProviderPin());
+	InCompilerContext.MovePinLinksToIntermediate(*FindPinChecked(UDataLinkEdNode::PN_Output, EGPD_Output), *AsyncRequestNode->GetOutputDataPin());
 
 	UEnum* const ResultEnum = StaticEnum<EDataLinkExecutionResult>();
 
@@ -140,8 +137,8 @@ void UK2Node_DataLinkRequest::ExpandNode(FKismetCompilerContext& InCompilerConte
 		UEdGraphPin* SuccessOutputPin = SwitchNode->FindPinChecked(ResultEnum->GetNameStringByValue(static_cast<uint64>(EDataLinkExecutionResult::Succeeded)), EGPD_Output);
 		UEdGraphPin* FailureOutputPin = SwitchNode->FindPinChecked(ResultEnum->GetNameStringByValue(static_cast<uint64>(EDataLinkExecutionResult::Failed)), EGPD_Output);
 
-		InCompilerContext.MovePinLinksToIntermediate(*FindPinChecked(UK2Node_DataLinkRequest::PN_Success), *SuccessOutputPin);
-		InCompilerContext.MovePinLinksToIntermediate(*FindPinChecked(UK2Node_DataLinkRequest::PN_Failure), *FailureOutputPin);
+		InCompilerContext.MovePinLinksToIntermediate(*FindPinChecked(PN_Success), *SuccessOutputPin);
+		InCompilerContext.MovePinLinksToIntermediate(*FindPinChecked(PN_Failure), *FailureOutputPin);
 	}
 
 	BreakAllNodeLinks();
diff --git a/Engine/Plugins/Experimental/DataLink/Source/DataLinkEdGraph/Public/Blueprint/K2Node_DataLinkRequest.h b/Engine/Plugins/Experimental/DataLink/Source/DataLinkEdGraph/Public/Blueprint/K2Node_DataLinkRequest.h
--- a/Engine/Plugins/Experimental/DataLink/Source/DataLinkEdGraph/Public/Blueprint/K2Node_DataLinkRequest.h
+++ b/Engine/Plugins/Experimental/DataLink/Source/DataLinkEdGraph/Public/Blueprint/K2Node_DataLinkRequest.h
@@ -8,7 +8,7 @@
 class UDataLinkGraph;
 
 UCLASS(MinimalAPI)
-class UK2Node_DataLinkRequest : public UK2Node
+class UE_DEPRECATED(5.7, "Data Link Request deprecated. Use Data Link Executor Object instead.") UK2Node_DataLinkRequest : public UK2Node
 {
 	GENERATED_BODY()
 
@@ -19,15 +19,12 @@ public:
 	static const FLazyName PN_Success;
 	static const FLazyName PN_Failure;
 
-	UEdGraphPin* GetDataLinkInstancePin() const;
-	UEdGraphPin* GetExecutionContextPin() const;
-	UEdGraphPin* GetDataLinkSinkProviderPin() const;
-	UEdGraphPin* GetOutputPin() const;
-
 	//~ Begin UEdGraphNode
 	virtual void AllocateDefaultPins() override;
 	virtual bool IsCompatibleWithGraph(const UEdGraph* InTargetGraph) const override;
 	virtual FText GetNodeTitle(ENodeTitleType::Type InTitleType) const override;
+	virtual bool IsDeprecated() const override;
+	virtual FEdGraphNodeDeprecationResponse GetDeprecationResponse(EEdGraphNodeDeprecationType InDeprecationType) const override;
 	//~ End UEdGraphNode
 
 	//~ Begin UK2Node
