Motion Design: data link auto compile for old graphs that don't have output pin name set

diff --git a/Engine/Plugins/Experimental/DataLink/Source/DataLink/Private/DataLinkDelegates.cpp b/Engine/Plugins/Experimental/DataLink/Source/DataLink/Private/DataLinkDelegates.cpp
new file mode 100644
--- /dev/null
+++ b/Engine/Plugins/Experimental/DataLink/Source/DataLink/Private/DataLinkDelegates.cpp
@@ -0,0 +1,9 @@
+// Copyright Epic Games, Inc. All Rights Reserved.
+
+#include "DataLinkDelegates.h"
+
+namespace UE::DataLink
+{
+	FOnRequestGraphCompilation OnRequestCompilation;
+
+} // UE::DataLink
diff --git a/Engine/Plugins/Experimental/DataLink/Source/DataLink/Private/DataLinkExecutor.cpp b/Engine/Plugins/Experimental/DataLink/Source/DataLink/Private/DataLinkExecutor.cpp
--- a/Engine/Plugins/Experimental/DataLink/Source/DataLink/Private/DataLinkExecutor.cpp
+++ b/Engine/Plugins/Experimental/DataLink/Source/DataLink/Private/DataLinkExecutor.cpp
@@ -240,7 +240,14 @@ bool FDataLinkExecutor::ValidateRun()
 {
 	if (!Instance.DataLinkGraph)
 	{
-		UE_LOG(LogDataLink, Error, TEXT("[%s] Invalid Data Link specified!"), GetContextName().GetData());
+		UE_LOG(LogDataLink, Error, TEXT("[%.*s] Invalid Data Link specified!")
+			, GetContextName().Len()
+			, GetContextName().GetData());
+		return false;
+	}
+
+	if (!ValidateGraphOutputNode())
+	{
 		return false;
 	}
 
@@ -248,7 +255,8 @@ bool FDataLinkExecutor::ValidateRun()
 
 	if (Instance.InputData.Num() != InputPinCount)
 	{
-		UE_LOG(LogDataLink, Error, TEXT("[%s] DataLink Graph '%s' requires %d inputs but %d were provided")
+		UE_LOG(LogDataLink, Error, TEXT("[%.*s] DataLink Graph '%s' requires %d inputs but %d were provided")
+			, GetContextName().Len()
 			, GetContextName().GetData()
 			, *Instance.DataLinkGraph->GetName()
 			, InputPinCount
@@ -266,7 +274,8 @@ bool FDataLinkExecutor::ValidateRun()
 
 			if (InPinReference.Pin->Struct && InPinReference.Pin->Struct != Instance.InputData[CurrentPinIndex].Data.GetScriptStruct())
 			{
-				UE_LOG(LogDataLink, Error, TEXT("[%s] Input pin (%s': input struct '%s', owner '%s') is not compatible with input data view '%s'")
+				UE_LOG(LogDataLink, Error, TEXT("[%.*s] Input pin (%s': input struct '%s', owner '%s') is not compatible with input data view '%s'")
+					, GetContextName().Len()
 					, GetContextName().GetData()
 					, *InPinReference.Pin->Name.ToString()
 					, *InPinReference.Pin->Struct->GetName()
@@ -281,6 +290,63 @@ bool FDataLinkExecutor::ValidateRun()
 	return bResult;
 }
 
+bool FDataLinkExecutor::ValidateGraphOutputNode()
+{
+	using namespace UE::DataLink;
+
+	// Already verified to be valid
+	check(Instance.DataLinkGraph);
+
+	const UDataLinkNode* OutputNode = Instance.DataLinkGraph->GetOutputNode();
+	if (!OutputNode)
+	{
+		UE_LOG(LogDataLink, Error, TEXT("[%.*s] Invalid Data Link output node! Graph: %s")
+			, GetContextName().Len()
+			, GetContextName().GetData()
+			, *Instance.DataLinkGraph->GetName());
+		return false;
+	}
+
+	// If there's no valid output pin with the stored name then try to compile graph or use the first output pin name.
+	// This should only happen for old graphs before output pin name was introduced
+	const FName OutputPinName = Instance.DataLinkGraph->GetOutputPinName();
+	if (!OutputNode->GetOutputPins().Contains(OutputPinName))
+	{
+		// Todo: consider possible option in executor arguments to not do any compilation even if the graph is in a bad state and just fail execution 
+		if (UE::DataLink::OnRequestCompilation.IsBound())
+		{
+			UE_LOG(LogDataLink, Warning, TEXT("[%.*s] Output pin name '%s' was not found as a valid pin name in output node '%s' in Data Link Graph '%s'. Recompiling graph...")
+				, GetContextName().Len()
+				, GetContextName().GetData()
+				, *OutputPinName.ToString()
+				, *OutputNode->GetName()
+				, *Instance.DataLinkGraph->GetName());
+
+			const EGraphCompileStatus Result = UE::DataLink::OnRequestCompilation.Execute(Instance.DataLinkGraph);
+			if (Result != EGraphCompileStatus::UpToDate)
+			{
+				UE_LOG(LogDataLink, Error, TEXT("[%.*s] Graph '%s' failed to compile. Recompile graph manually.")
+					, GetContextName().Len()
+					, GetContextName().GetData()
+					, *Instance.DataLinkGraph->GetName());
+				return false;
+			}
+		}
+		else
+		{
+			UE_LOG(LogDataLink, Error, TEXT("[%.*s] Output pin name '%s' was not found as a valid pin name in output node '%s' in Data Link Graph '%s'. Recompile Graph manually.")
+				, GetContextName().Len()
+				, GetContextName().GetData()
+				, *OutputPinName.ToString()
+				, *OutputNode->GetName()
+				, *Instance.DataLinkGraph->GetName());
+			return false;
+		}
+	}
+
+	return true;
+}
+
 bool FDataLinkExecutor::ValidateInputPins(TConstArrayView<FDataLinkPinReference> InInputPins, TConstArrayView<FConstStructView> InInputDataViews) const
 {
 	if (!ensureAlways(InInputPins.Num() == InInputDataViews.Num()))
diff --git a/Engine/Plugins/Experimental/DataLink/Source/DataLink/Public/DataLinkDelegates.h b/Engine/Plugins/Experimental/DataLink/Source/DataLink/Public/DataLinkDelegates.h
--- a/Engine/Plugins/Experimental/DataLink/Source/DataLink/Public/DataLinkDelegates.h
+++ b/Engine/Plugins/Experimental/DataLink/Source/DataLink/Public/DataLinkDelegates.h
@@ -5,9 +5,28 @@
 #include "Delegates/Delegate.h"
 #include "Delegates/DelegateCombinations.h"
 
+#define UE_API DATALINK_API
+
+namespace UE::DataLink
+{
+	enum class EGraphCompileStatus : uint8;
+}
+
 class FDataLinkExecutor;
+class UDataLinkGraph;
 enum class EDataLinkExecutionResult : uint8;
 struct FConstStructView;
 
 DECLARE_DELEGATE_TwoParams(FOnDataLinkOutputData, const FDataLinkExecutor&, FConstStructView);
 DECLARE_DELEGATE_TwoParams(FOnDataLinkExecutionFinished, const FDataLinkExecutor&, EDataLinkExecutionResult);
+
+namespace UE::DataLink
+{
+	DECLARE_DELEGATE_RetVal_OneParam(EGraphCompileStatus, FOnRequestGraphCompilation, UDataLinkGraph*);
+
+	/** Delegate to request to the editor module to compile a data link graph */
+	extern UE_API FOnRequestGraphCompilation OnRequestCompilation;
+
+} // UE::DataLink
+
+#undef UE_API
diff --git a/Engine/Plugins/Experimental/DataLink/Source/DataLink/Public/DataLinkEnums.h b/Engine/Plugins/Experimental/DataLink/Source/DataLink/Public/DataLinkEnums.h
--- a/Engine/Plugins/Experimental/DataLink/Source/DataLink/Public/DataLinkEnums.h
+++ b/Engine/Plugins/Experimental/DataLink/Source/DataLink/Public/DataLinkEnums.h
@@ -5,6 +5,18 @@
 #include "HAL/Platform.h"
 #include "DataLinkEnums.generated.h"
 
+namespace UE::DataLink
+{
+	enum class EGraphCompileStatus : uint8
+	{
+		Unknown,
+		Warning,
+		Error,
+		UpToDate,
+	};
+
+} // UE::DataLink
+
 UENUM()
 enum class EDataLinkExecutionResult : uint8
 {
diff --git a/Engine/Plugins/Experimental/DataLink/Source/DataLink/Public/DataLinkExecutor.h b/Engine/Plugins/Experimental/DataLink/Source/DataLink/Public/DataLinkExecutor.h
--- a/Engine/Plugins/Experimental/DataLink/Source/DataLink/Public/DataLinkExecutor.h
+++ b/Engine/Plugins/Experimental/DataLink/Source/DataLink/Public/DataLinkExecutor.h
@@ -91,6 +91,9 @@ private:
 	/** Called when running to ensure all the parameters are properly set. Returns true if run can be done, false otherwise */
 	bool ValidateRun();
 
+	/** Called to ensure the graph's output node is in a valid state. Returns true if run can be done, false otherwise */
+	bool ValidateGraphOutputNode();
+
 	/** Called in Pin execution to make sure the pins and data views match in compatibility */
 	bool ValidateInputPins(TConstArrayView<FDataLinkPinReference> InInputPins, TConstArrayView<FConstStructView> InInputDataViews) const;
 
diff --git a/Engine/Plugins/Experimental/DataLink/Source/DataLink/Public/DataLinkGraph.h b/Engine/Plugins/Experimental/DataLink/Source/DataLink/Public/DataLinkGraph.h
--- a/Engine/Plugins/Experimental/DataLink/Source/DataLink/Public/DataLinkGraph.h
+++ b/Engine/Plugins/Experimental/DataLink/Source/DataLink/Public/DataLinkGraph.h
@@ -6,6 +6,8 @@
 #include "UObject/Object.h"
 #include "DataLinkGraph.generated.h"
 
+#define UE_API DATALINK_API
+
 class UDataLinkNode;
 struct FConstStructView;
 struct FDataLinkPinReference;
@@ -24,13 +26,13 @@ class UDataLinkGraph : public UObject
 
 public:
 	/** Counts the number of input pins within the input node list that are not linked to other nodes */
-	DATALINK_API int32 GetInputPinCount() const;
+	UE_API int32 GetInputPinCount() const;
 
 	/** Retrieves the Input Pins from the input nodes that are not linked to other nodes */
-	DATALINK_API TArray<FDataLinkPinReference> GetInputPins() const;
+	UE_API TArray<FDataLinkPinReference> GetInputPins() const;
 
 	/** Iterates the input pins from the input nodes that are not linked to other nodes. These serve as the 'input pin' for the graph */
-	DATALINK_API bool ForEachInputPin(TFunctionRef<bool(FDataLinkPinReference)> InFunction) const;
+	UE_API bool ForEachInputPin(TFunctionRef<bool(FDataLinkPinReference)> InFunction) const;
 
 	TConstArrayView<const UDataLinkNode*> GetInputNodes() const
 	{
@@ -61,7 +63,7 @@ public:
 
 private:
 #if WITH_EDITOR
-	DATALINK_API static TMulticastDelegate<void(UDataLinkGraph*)> OnGraphCompiledDelegate;
+	UE_API static TMulticastDelegate<void(UDataLinkGraph*)> OnGraphCompiledDelegate;
 #endif
 
 	/** All the compiled nodes present in this graph */
@@ -86,3 +88,5 @@ private:
 	TObjectPtr<UEdGraph> EdGraph;
 #endif
 };
+
+#undef UE_API
diff --git a/Engine/Plugins/Experimental/DataLink/Source/DataLinkEditor/Private/Compiler/DataLinkGraphCompileStatus.h b/Engine/Plugins/Experimental/DataLink/Source/DataLinkEditor/Private/Compiler/DataLinkGraphCompileStatus.h
deleted file mode 100644
diff --git a/Engine/Plugins/Experimental/DataLink/Source/DataLinkEditor/Private/Compiler/DataLinkGraphCompiler.cpp b/Engine/Plugins/Experimental/DataLink/Source/DataLinkEditor/Private/Compiler/DataLinkGraphCompiler.cpp
--- a/Engine/Plugins/Experimental/DataLink/Source/DataLinkEditor/Private/Compiler/DataLinkGraphCompiler.cpp
+++ b/Engine/Plugins/Experimental/DataLink/Source/DataLinkEditor/Private/Compiler/DataLinkGraphCompiler.cpp
@@ -3,7 +3,6 @@
 #include "DataLinkGraphCompiler.h"
 #include "DataLinkEditorLog.h"
 #include "DataLinkGraph.h"
-#include "DataLinkGraphCompileStatus.h"
 #include "DataLinkNode.h"
 #include "EdGraph/EdGraphPin.h"
 #include "Nodes/DataLinkEdNode.h"
@@ -14,19 +13,21 @@ FDataLinkGraphCompiler::FDataLinkGraphCompiler(UDataLinkGraph* InDataLinkGraph)
 {
 }
 
-EDataLinkGraphCompileStatus FDataLinkGraphCompiler::Compile()
+UE::DataLink::EGraphCompileStatus FDataLinkGraphCompiler::Compile()
 {
+	using namespace UE::DataLink;
+
 	if (!DataLinkGraph)
 	{
 		UE_LOG(LogDataLinkEditor, Error, TEXT("Compilation failed. Data Link Graph is invalid."));
-		return EDataLinkGraphCompileStatus::Error;
+		return EGraphCompileStatus::Error;
 	}
 
 	DataLinkEdGraph = Cast<UDataLinkEdGraph>(DataLinkGraph->GetEdGraph());
 	if (!DataLinkEdGraph)
 	{
 		UE_LOG(LogDataLinkEditor, Error, TEXT("Compilation failed. Data Link Ed Graph is invalid in graph '%s'"), *DataLinkGraph->GetName());
-		return EDataLinkGraphCompileStatus::Error;
+		return EGraphCompileStatus::Error;
 	}
 
 	CleanExistingGraph();
@@ -35,11 +36,11 @@ EDataLinkGraphCompileStatus FDataLinkGraphCompiler::Compile()
 	{
 		// Compilation failed. Clean graph again from what had started compiling
     	CleanExistingGraph();
-    	return EDataLinkGraphCompileStatus::Error;
+    	return EGraphCompileStatus::Error;
 	}
 
 	UDataLinkGraph::OnGraphCompiledDelegate.Broadcast(DataLinkGraph);
-	return EDataLinkGraphCompileStatus::Good;
+	return EGraphCompileStatus::UpToDate;
 }
 
 void FDataLinkGraphCompiler::CleanExistingGraph()
@@ -54,6 +55,7 @@ bool FDataLinkGraphCompiler::CompileNodes()
 	if (DataLinkEdGraph->Nodes.IsEmpty())
 	{
 		// Nothing further to compile. Unknown Data Link Graph
+		UE_LOG(LogDataLinkEditor, Error, TEXT("Compilation failed. No graph node were found! Graph: %s"), *DataLinkGraph->GetName());
 		return false;
 	}
 
diff --git a/Engine/Plugins/Experimental/DataLink/Source/DataLinkEditor/Private/Compiler/DataLinkGraphCompiler.h b/Engine/Plugins/Experimental/DataLink/Source/DataLinkEditor/Private/Compiler/DataLinkGraphCompiler.h
--- a/Engine/Plugins/Experimental/DataLink/Source/DataLinkEditor/Private/Compiler/DataLinkGraphCompiler.h
+++ b/Engine/Plugins/Experimental/DataLink/Source/DataLinkEditor/Private/Compiler/DataLinkGraphCompiler.h
@@ -4,17 +4,23 @@
 
 #include "DataLinkEdGraph.h"
 
+namespace UE::DataLink
+{
+	enum class EGraphCompileStatus : uint8;
+}
+
 class UDataLinkEdNode;
 class UDataLinkGraph;
 class UDataLinkNode;
 enum class EDataLinkGraphCompileStatus : uint8;
 
+/** Compiles a data link graph's ed graph and writes its results to the runtime data of the data link graph */
 class FDataLinkGraphCompiler
 {
 public:
 	explicit FDataLinkGraphCompiler(UDataLinkGraph* InDataLinkGraph);
 
-	EDataLinkGraphCompileStatus Compile();
+	UE::DataLink::EGraphCompileStatus Compile();
 
 private:
 	void CleanExistingGraph();
diff --git a/Engine/Plugins/Experimental/DataLink/Source/DataLinkEditor/Private/Compiler/DataLinkGraphCompilerTool.cpp b/Engine/Plugins/Experimental/DataLink/Source/DataLinkEditor/Private/Compiler/DataLinkGraphCompilerTool.cpp
--- a/Engine/Plugins/Experimental/DataLink/Source/DataLinkEditor/Private/Compiler/DataLinkGraphCompilerTool.cpp
+++ b/Engine/Plugins/Experimental/DataLink/Source/DataLinkEditor/Private/Compiler/DataLinkGraphCompilerTool.cpp
@@ -2,6 +2,7 @@
 
 #include "DataLinkGraphCompilerTool.h"
 #include "DataLinkEditorStyle.h"
+#include "DataLinkEnums.h"
 #include "DataLinkGraphAssetEditor.h"
 #include "DataLinkGraphAssetToolkit.h"
 #include "DataLinkGraphCommands.h"
@@ -12,22 +13,22 @@
 
 #define LOCTEXT_NAMESPACE "DataLinkGraphCompilerToolkit"
 
-namespace UE::DataLinkEditor::Private
+namespace UE::DataLink::Editor::Private
 {
-	const TMap<EDataLinkGraphCompileStatus, FLazyName> CompileStatusBackgrounds =
+	const TMap<EGraphCompileStatus, FLazyName> CompileStatusBackgrounds =
 	{
-		{ EDataLinkGraphCompileStatus::Unknown , TEXT("CompileStatus.Background.Unknown") },
-		{ EDataLinkGraphCompileStatus::Warning , TEXT("CompileStatus.Background.Warning") },
-		{ EDataLinkGraphCompileStatus::Error   , TEXT("CompileStatus.Background.Error")  },
-		{ EDataLinkGraphCompileStatus::Good    , TEXT("CompileStatus.Background.Good")   },
+		{ EGraphCompileStatus::Unknown , TEXT("CompileStatus.Background.Unknown") },
+		{ EGraphCompileStatus::Warning , TEXT("CompileStatus.Background.Warning") },
+		{ EGraphCompileStatus::Error   , TEXT("CompileStatus.Background.Error")  },
+		{ EGraphCompileStatus::UpToDate, TEXT("CompileStatus.Background.Good")   },
 	};
 
-	const TMap<EDataLinkGraphCompileStatus, FLazyName> CompileStatusOverlays =
+	const TMap<EGraphCompileStatus, FLazyName> CompileStatusOverlays =
 	{
-		{ EDataLinkGraphCompileStatus::Unknown , TEXT("CompileStatus.Overlay.Unknown") },
-		{ EDataLinkGraphCompileStatus::Warning , TEXT("CompileStatus.Overlay.Warning") },
-		{ EDataLinkGraphCompileStatus::Error   , TEXT("CompileStatus.Overlay.Error")  },
-		{ EDataLinkGraphCompileStatus::Good    , TEXT("CompileStatus.Overlay.Good")   },
+		{ EGraphCompileStatus::Unknown , TEXT("CompileStatus.Overlay.Unknown") },
+		{ EGraphCompileStatus::Warning , TEXT("CompileStatus.Overlay.Warning") },
+		{ EGraphCompileStatus::Error   , TEXT("CompileStatus.Overlay.Error")  },
+		{ EGraphCompileStatus::UpToDate, TEXT("CompileStatus.Overlay.Good")   },
 	};
 }
 
@@ -43,12 +44,12 @@ void FDataLinkGraphCompilerTool::ExtendMenu(UToolMenu* InMenu)
 
 FDataLinkGraphCompilerTool::FDataLinkGraphCompilerTool(UDataLinkGraphAssetEditor* InAssetEditor)
 	: AssetEditor(InAssetEditor)
-	, LastCompiledStatus(EDataLinkGraphCompileStatus::Unknown)
+	, LastCompiledStatus(UE::DataLink::EGraphCompileStatus::Unknown)
 {
 	UDataLinkEdGraph* EdGraph = AssetEditor->GetDataLinkEdGraph();
 	if (EdGraph && EdGraph->IsCompiledGraphUpToDate())
 	{
-		LastCompiledStatus = EDataLinkGraphCompileStatus::Good;
+		LastCompiledStatus = UE::DataLink::EGraphCompileStatus::UpToDate;
 	}
 }
 
@@ -88,16 +89,18 @@ void FDataLinkGraphCompilerTool::ExtendDynamicCompilerSection(FToolMenuSection&
 
 FSlateIcon FDataLinkGraphCompilerTool::GetCompileIcon() const
 {
-	EDataLinkGraphCompileStatus CompileStatus = LastCompiledStatus;
+	using namespace UE::DataLink;
+
+	EGraphCompileStatus CompileStatus = LastCompiledStatus;
 
 	UDataLinkEdGraph* EdGraph = AssetEditor->GetDataLinkEdGraph();
 	if (EdGraph && !EdGraph->IsCompiledGraphUpToDate())
 	{
-		CompileStatus = EDataLinkGraphCompileStatus::Unknown;
+		CompileStatus = EGraphCompileStatus::Unknown;
 	}
 
-	const FName Background = UE::DataLinkEditor::Private::CompileStatusBackgrounds[CompileStatus]; 
-	const FName Overlay = UE::DataLinkEditor::Private::CompileStatusOverlays[CompileStatus]; 
+	const FName Background = Editor::Private::CompileStatusBackgrounds[CompileStatus]; 
+	const FName Overlay = Editor::Private::CompileStatusOverlays[CompileStatus]; 
 
 	return FSlateIcon(FDataLinkEditorStyle::Get().GetStyleSetName(), Background, Background, Overlay);
 }
diff --git a/Engine/Plugins/Experimental/DataLink/Source/DataLinkEditor/Private/Compiler/DataLinkGraphCompilerTool.h b/Engine/Plugins/Experimental/DataLink/Source/DataLinkEditor/Private/Compiler/DataLinkGraphCompilerTool.h
--- a/Engine/Plugins/Experimental/DataLink/Source/DataLinkEditor/Private/Compiler/DataLinkGraphCompilerTool.h
+++ b/Engine/Plugins/Experimental/DataLink/Source/DataLinkEditor/Private/Compiler/DataLinkGraphCompilerTool.h
@@ -2,10 +2,14 @@
 
 #pragma once
 
-#include "DataLinkGraphCompileStatus.h"
 #include "Templates/SharedPointer.h"
 #include "UObject/ObjectPtr.h"
 
+namespace UE::DataLink
+{
+	enum class EGraphCompileStatus : uint8;
+}
+
 class FUICommandList;
 class UDataLinkGraphAssetEditor;
 class UToolMenu;
@@ -30,6 +34,6 @@ private:
 
 	TObjectPtr<UDataLinkGraphAssetEditor> AssetEditor;
 
-	EDataLinkGraphCompileStatus LastCompiledStatus;
+	UE::DataLink::EGraphCompileStatus LastCompiledStatus;
 };
 
diff --git a/Engine/Plugins/Experimental/DataLink/Source/DataLinkEditor/Private/DataLinkEditorModule.cpp b/Engine/Plugins/Experimental/DataLink/Source/DataLinkEditor/Private/DataLinkEditorModule.cpp
--- a/Engine/Plugins/Experimental/DataLink/Source/DataLinkEditor/Private/DataLinkEditorModule.cpp
+++ b/Engine/Plugins/Experimental/DataLink/Source/DataLinkEditor/Private/DataLinkEditorModule.cpp
@@ -1,6 +1,8 @@
 // Copyright Epic Games, Inc. All Rights Reserved.
 
 #include "DataLinkEditorModule.h"
+#include "Compiler/DataLinkGraphCompiler.h"
+#include "DataLinkDelegates.h"
 #include "DataLinkEditorLog.h"
 #include "DataLinkGraphCommands.h"
 #include "DataLinkInstance.h"
@@ -17,12 +19,21 @@ void FDataLinkEditorModule::StartupModule()
 {
 	FDataLinkGraphCommands::Register();
 	RegisterDetailCustomizations();
+
+	UE::DataLink::OnRequestCompilation.BindLambda(
+		[](UDataLinkGraph* InGraphToCompile)->UE::DataLink::EGraphCompileStatus
+		{
+			FDataLinkGraphCompiler GraphCompiler(InGraphToCompile);
+			return GraphCompiler.Compile();
+		});
 }
 
 void FDataLinkEditorModule::ShutdownModule()
 {
 	FDataLinkGraphCommands::Unregister();
 	UnregisterDetailCustomizations();
+
+	UE::DataLink::OnRequestCompilation.Unbind();
 }
 
 void FDataLinkEditorModule::RegisterDetailCustomizations()
