Motion Design: fixed comment for clarity

diff --git a/Engine/Plugins/Experimental/DataLink/Source/DataLink/Private/DataLinkExecutor.cpp b/Engine/Plugins/Experimental/DataLink/Source/DataLink/Private/DataLinkExecutor.cpp
--- a/Engine/Plugins/Experimental/DataLink/Source/DataLink/Private/DataLinkExecutor.cpp
+++ b/Engine/Plugins/Experimental/DataLink/Source/DataLink/Private/DataLinkExecutor.cpp
@@ -307,8 +307,8 @@ bool FDataLinkExecutor::ValidateGraphOutputNode()
 		return false;
 	}
 
-	// If there's no valid output pin with the stored name then try to compile graph or use the first output pin name.
-	// This should only happen for old graphs before output pin name was introduced
+	// If there's no valid output pin with the stored name then try to compile graph.
+	// This should only happen for old graphs before output pin name was introduced.
 	const FName OutputPinName = Instance.DataLinkGraph->GetOutputPinName();
 	if (!OutputNode->GetOutputPins().Contains(OutputPinName))
 	{
