Motion Design Data Link: fix issue with pin corruption when undoing/redoing operations involving node creation/deletion and pin linking

#FailedValidations Preflight

diff --git a/Engine/Plugins/VirtualProduction/DataLink/Source/DataLinkEdGraph/Private/DataLinkEdGraph.cpp b/Engine/Plugins/VirtualProduction/DataLink/Source/DataLinkEdGraph/Private/DataLinkEdGraph.cpp
--- a/Engine/Plugins/VirtualProduction/DataLink/Source/DataLinkEdGraph/Private/DataLinkEdGraph.cpp
+++ b/Engine/Plugins/VirtualProduction/DataLink/Source/DataLinkEdGraph/Private/DataLinkEdGraph.cpp
@@ -23,17 +23,10 @@ void UDataLinkEdGraph::InitializeNodes()
 {
 	for (UEdGraphNode* Node : Nodes)
 	{
-		UDataLinkEdNode* DataLinkEdNode = Cast<UDataLinkEdNode>(Node);
-		if (!DataLinkEdNode)
+		if (UDataLinkEdNode* DataLinkEdNode = Cast<UDataLinkEdNode>(Node))
 		{
-			continue;
-		}
-
-		DataLinkEdNode->UpdateMetadata();
-
-		if (DataLinkEdNode->RequiresPinRecreation())
-		{
-			DataLinkEdNode->RecreatePins();	
+			DataLinkEdNode->UpdateMetadata();
+			DataLinkEdNode->RefreshPins();
 		}
 	}
 }
diff --git a/Engine/Plugins/VirtualProduction/DataLink/Source/DataLinkEdGraph/Private/Nodes/DataLinkEdNode.cpp b/Engine/Plugins/VirtualProduction/DataLink/Source/DataLinkEdGraph/Private/Nodes/DataLinkEdNode.cpp
--- a/Engine/Plugins/VirtualProduction/DataLink/Source/DataLinkEdGraph/Private/Nodes/DataLinkEdNode.cpp
+++ b/Engine/Plugins/VirtualProduction/DataLink/Source/DataLinkEdGraph/Private/Nodes/DataLinkEdNode.cpp
@@ -14,6 +14,13 @@ const FLazyName UDataLinkEdNode::PN_Output = TEXT("Output");
 
 namespace UE::DataLinkEdGraph::Private
 {
+	bool IsMatchingPin(TNotNull<const UEdGraphPin*> InPin, const FDataLinkPin& InDataLinkPin)
+	{
+		return InPin->PinName == InDataLinkPin.Name
+			&& InPin->PinType.PinCategory == UDataLinkEdGraphSchema::PC_Data
+			&& InPin->PinType.PinSubCategoryObject == InDataLinkPin.Struct;
+	}
+
 	bool FindLinkedNode(const UEdGraphPin* InPin, const UDataLinkEdNode*& OutLinkedNode, const UEdGraphPin*& OutLinkedPin)
 	{
 		if (!InPin || InPin->LinkedTo.IsEmpty())
@@ -97,17 +104,14 @@ bool UDataLinkEdNode::RequiresPinRecreation() const
 		return true;
 	}
 
-	auto HasMismatchingPin = [this](TConstArrayView<FDataLinkPin> InTemplatePins, EEdGraphPinDirection InDirection)
+	auto HasMismatchingPin = [This=this](TConstArrayView<FDataLinkPin> InTemplatePins, EEdGraphPinDirection InDirection)
 		{
 			for (const FDataLinkPin& TemplatePin : InTemplatePins)
 			{
-				const UEdGraphPin* const FoundPin = FindPin(TemplatePin.Name, InDirection);
+				const UEdGraphPin* const FoundPin = This->FindPin(TemplatePin.Name, InDirection);
 
-				// Require pin recreation if the found pin that is supposed to match the template pin by name is not a data pin
-				// or has a mismatching struct
-				if (!FoundPin
-					|| FoundPin->PinType.PinCategory != UDataLinkEdGraphSchema::PC_Data
-					|| FoundPin->PinType.PinSubCategoryObject != TemplatePin.Struct)
+				// Require pin recreation if there's no found pin for the template data link pin or if it mismatches from it
+				if (!FoundPin || !UE::DataLinkEdGraph::Private::IsMatchingPin(FoundPin, TemplatePin))
 				{
 					return true;
 				}
@@ -167,7 +171,7 @@ void UDataLinkEdNode::ReconstructNode()
 {
 	Super::ReconstructNode();
 	UpdateMetadata();
-	RecreatePins();
+	RefreshPins();
 }
 
 bool UDataLinkEdNode::CanCreateUnderSpecifiedSchema(const UEdGraphSchema* InSchema) const
@@ -201,7 +205,7 @@ void UDataLinkEdNode::PostEditUndo()
 {
 	Super::PostEditUndo();
 	UpdateMetadata();
-	RecreatePins();
+	NotifyNodeChanged();
 }
 
 void UDataLinkEdNode::PostEditChangeChainProperty(FPropertyChangedChainEvent& InPropertyChangedEvent)
@@ -214,7 +218,7 @@ void UDataLinkEdNode::PostEditChangeChainProperty(FPropertyChangedChainEvent& In
 		if (Property && Property->HasMetaData(UDataLinkEdNode::MD_InvalidatesNode))
 		{
 			UpdateMetadata();
-			RecreatePins();
+			RefreshPins();
 			break;
 		}
 	}
@@ -246,74 +250,66 @@ void UDataLinkEdNode::NotifyNodeChanged()
 	}
 }
 
-void UDataLinkEdNode::RecreatePins()
+void UDataLinkEdNode::RefreshPins()
 {
-	Modify(/*bAlwaysMarkDirty*/false);
-
-	const TArray<UEdGraphPin*> PinsCopy = Pins;
-
-	TArray<UEdGraphPin*> RemovedPins;
-	RemovedPins.Reserve(PinsCopy.Num());
-
-	// Remove all pins related to input data
-	for (UEdGraphPin* Pin : PinsCopy)
+	if (!RequiresPinRecreation())
 	{
-		if (Pin)
-		{
-			Pins.Remove(Pin);
-			RemovedPins.Add(Pin);
-		}
+		return;
 	}
 
-	AllocateDefaultPins();
+	TArray<FDataLinkPin> InputPins;
+	TArray<FDataLinkPin> OutputPins;
 
 	// Build Data Link Node Pins and create them on the Ed Node
 	if (TemplateNode)
 	{
-		TArray<FDataLinkPin> InputPins;
-		TArray<FDataLinkPin> OutputPins;
-
 		TemplateNode->BuildPins(InputPins, OutputPins);
+	}
 
-		CreatePins(EGPD_Input, InputPins);
-		CreatePins(EGPD_Output, OutputPins);
+	bool bPinsChanged = false;
+	bPinsChanged |= SyncPins(EGPD_Input, InputPins);
+	bPinsChanged |= SyncPins(EGPD_Output, OutputPins);
+
+	// Refresh the UI for the graph so the pin changes show up
+	if (bPinsChanged)
+	{
+		NotifyNodeChanged();
 	}
+}
 
-	// Rewire the pins to remove and matching new pins
-	if (const UEdGraphSchema* Schema = GetSchema())
+bool UDataLinkEdNode::SyncPins(EEdGraphPinDirection InPinDirection, TConstArrayView<FDataLinkPin> InTemplatePins)
+{
+	bool bPinsChanged = false;
+
+	// Remove pins that are no longer present in the provided template pins, or that mismatch the template pin
+	TArray<UEdGraphPin*> PinsCopy = Pins;
+	for (UEdGraphPin* Pin : PinsCopy)
 	{
-		for (UEdGraphPin* RemovedPin : RemovedPins)
+		if (!Pin || Pin->Direction != InPinDirection)
 		{
-			RemovedPin->Modify(/*bAlwaysMarkDirty*/false);
-
-			if (UEdGraphPin* NewPin = FindPin(RemovedPin->PinName, RemovedPin->Direction))
-			{
-				Schema->MovePinLinks(*RemovedPin, *NewPin);
-			}
+			continue;
+		}
 
-			RemovedPin->MarkAsGarbage();
-			OnPinRemoved(RemovedPin);
+		const FDataLinkPin* const TemplatePin = InTemplatePins.FindByKey(Pin->PinName);
+		if (!TemplatePin || !UE::DataLinkEdGraph::Private::IsMatchingPin(Pin, *TemplatePin))
+		{
+			RemovePin(Pin);
+			bPinsChanged = true;
 		}
 	}
 
-	// Refresh the UI for the graph so the pin changes show up
-	NotifyNodeChanged();
-}
-
-void UDataLinkEdNode::CreatePins(EEdGraphPinDirection InPinDirection, TConstArrayView<FDataLinkPin> InDataLinkPins)
-{
-	for (const FDataLinkPin& DataLinkPin : InDataLinkPins)
+	// Create missing pin links
+	for (const FDataLinkPin& DataLinkPin : InTemplatePins)
 	{
 		// No need to add if Pin already is there
 		if (!FindPin(DataLinkPin.Name, InPinDirection))
 		{
-			UEdGraphPin* const Pin = CreatePin(InPinDirection
-				, UDataLinkEdGraphSchema::PC_Data
-				, ConstCast(DataLinkPin.Struct)
-				, DataLinkPin.Name);
-
+			UEdGraphPin* const Pin = CreatePin(InPinDirection, UDataLinkEdGraphSchema::PC_Data, ConstCast(DataLinkPin.Struct), DataLinkPin.Name);
 			check(Pin);
 			Pin->PinFriendlyName = DataLinkPin.GetDisplayName();
+			bPinsChanged = true;
 		}
 	}
+
+	return bPinsChanged;
 }
diff --git a/Engine/Plugins/VirtualProduction/DataLink/Source/DataLinkEdGraph/Public/Nodes/DataLinkEdNode.h b/Engine/Plugins/VirtualProduction/DataLink/Source/DataLinkEdGraph/Public/Nodes/DataLinkEdNode.h
--- a/Engine/Plugins/VirtualProduction/DataLink/Source/DataLinkEdGraph/Public/Nodes/DataLinkEdNode.h
+++ b/Engine/Plugins/VirtualProduction/DataLink/Source/DataLinkEdGraph/Public/Nodes/DataLinkEdNode.h
@@ -59,10 +59,16 @@ public:
 	DATALINKEDGRAPH_API virtual bool RequiresPinRecreation() const;
 
 	/** Destroys existing pins and recreates new ones based on the underlying Template Node while also restoring pin links where available */
-	void RecreatePins();
+	void RefreshPins();
 
 private:
-	void CreatePins(EEdGraphPinDirection InPinDirection, TConstArrayView<FDataLinkPin> InDataLinkPins);
+	/**
+	 * Syncs the node pins to match the given data link pins
+	 * @param InPinDirection the direction of pins to consider
+	 * @param InTemplatePins the pins to match
+	 * @return true if the pin structure was changed
+	 */
+	bool SyncPins(EEdGraphPinDirection InPinDirection, TConstArrayView<FDataLinkPin> InTemplatePins);
 
 	/** Underlying Data Link Node this Editor Node represents */
 	UPROPERTY(VisibleAnywhere, Instanced, Category="Data Link")
diff --git a/Engine/Plugins/VirtualProduction/DataLink/Source/DataLinkEditor/Private/Graph/DataLinkGraphEditorTool.cpp b/Engine/Plugins/VirtualProduction/DataLink/Source/DataLinkEditor/Private/Graph/DataLinkGraphEditorTool.cpp
--- a/Engine/Plugins/VirtualProduction/DataLink/Source/DataLinkEditor/Private/Graph/DataLinkGraphEditorTool.cpp
+++ b/Engine/Plugins/VirtualProduction/DataLink/Source/DataLinkEditor/Private/Graph/DataLinkGraphEditorTool.cpp
@@ -140,7 +140,7 @@ void FDataLinkGraphEditorTool::DeleteSelectedNodes()
 		const TSet<UObject*>& SelectedNodes = GraphEditor->GetSelectedNodes();
 		NodesToDelete.Reserve(SelectedNodes.Num());
 
-		for (UObject* SelectedNode : GraphEditor->GetSelectedNodes())
+		for (UObject* SelectedNode : SelectedNodes)
 		{
 			UEdGraphNode* Node = Cast<UEdGraphNode>(SelectedNode);
 			if (Node && Node->CanUserDeleteNode())
