Data Link: trivial rename of http request 'response' output pin to 'response string' for clarity on type

diff --git a/Engine/Plugins/VirtualProduction/DataLink/Source/DataLinkHttp/Private/DataLinkHttpSource.cpp b/Engine/Plugins/VirtualProduction/DataLink/Source/DataLinkHttp/Private/DataLinkHttpSource.cpp
--- a/Engine/Plugins/VirtualProduction/DataLink/Source/DataLinkHttp/Private/DataLinkHttpSource.cpp
+++ b/Engine/Plugins/VirtualProduction/DataLink/Source/DataLinkHttp/Private/DataLinkHttpSource.cpp
@@ -23,7 +23,7 @@ void UDataLinkHttpSource::OnBuildPins(FDataLinkPinBuilder& Inputs, FDataLinkPinB
 		.SetStruct<FDataLinkHttpSettings>();
 
 	Outputs.Add(UE::DataLink::OutputDefault)
-		.SetDisplayName(LOCTEXT("OutputDisplay", "Response"))
+		.SetDisplayName(LOCTEXT("OutputDisplay", "Response String"))
 		.SetStruct<FDataLinkString>();
 }
 
