Motion Design: Fixed data link world gc leak crash on level change when data link preview is executed.

diff --git a/Engine/Plugins/VirtualProduction/DataLink/Source/DataLink/Private/DataLinkProcessorBlueprint.cpp b/Engine/Plugins/VirtualProduction/DataLink/Source/DataLink/Private/DataLinkProcessorBlueprint.cpp
--- a/Engine/Plugins/VirtualProduction/DataLink/Source/DataLink/Private/DataLinkProcessorBlueprint.cpp
+++ b/Engine/Plugins/VirtualProduction/DataLink/Source/DataLink/Private/DataLinkProcessorBlueprint.cpp
@@ -2,6 +2,7 @@
 
 #include "DataLinkProcessorBlueprint.h"
 #include "DataLinkExecutor.h"
+#include "Engine/World.h"
 
 UObject* UDataLinkProcessorBlueprint::GetContextObject() const
 {
@@ -33,9 +34,14 @@ void UDataLinkProcessorBlueprint::OnFinalize(const FDataLinkExecutor& InExecutor
 
 UWorld* UDataLinkProcessorBlueprint::GetWorld() const
 {
-	if (UObject* ContextObject = GetContextObject())
+	if (UObject* Context = GetContextObject())
 	{
-		return ContextObject->GetWorld();
+		if (UWorld* World = Context->GetWorld())
+		{
+			return World;
+		}
 	}
-	return nullptr;
+	// Blueprint processors could use latent action nodes like 'Delay' that require a valid world from GetWorld().
+	// If context object does not have a valid world, default to GWorld
+	return GWorld;
 }
diff --git a/Engine/Plugins/VirtualProduction/DataLink/Source/DataLink/Private/Nodes/Script/DataLinkScriptNode.cpp b/Engine/Plugins/VirtualProduction/DataLink/Source/DataLink/Private/Nodes/Script/DataLinkScriptNode.cpp
--- a/Engine/Plugins/VirtualProduction/DataLink/Source/DataLink/Private/Nodes/Script/DataLinkScriptNode.cpp
+++ b/Engine/Plugins/VirtualProduction/DataLink/Source/DataLink/Private/Nodes/Script/DataLinkScriptNode.cpp
@@ -4,6 +4,7 @@
 #include "Blueprint/BlueprintExceptionInfo.h"
 #include "DataLinkCoreTypes.h"
 #include "DataLinkExecutor.h"
+#include "Engine/World.h"
 
 #define LOCTEXT_NAMESPACE "DataLinkScriptNode"
 
@@ -85,9 +86,14 @@ UWorld* UDataLinkScriptNode::GetWorld() const
 {
 	if (UObject* Context = GetContextObject())
 	{
-		return Context->GetWorld();
+		if (UWorld* World = Context->GetWorld())
+		{
+			return World;
+		}
 	}
-	return nullptr;
+	// Blueprint nodes could use latent action nodes like 'Delay' that require a valid world from GetWorld().
+	// If context object does not have a valid world, default to GWorld
+	return GWorld;
 }
 
 UObject* UDataLinkScriptNode::GetContextObject() const
diff --git a/Engine/Plugins/VirtualProduction/DataLink/Source/DataLinkEditor/Private/Preview/DataLinkPreviewTool.cpp b/Engine/Plugins/VirtualProduction/DataLink/Source/DataLinkEditor/Private/Preview/DataLinkPreviewTool.cpp
--- a/Engine/Plugins/VirtualProduction/DataLink/Source/DataLinkEditor/Private/Preview/DataLinkPreviewTool.cpp
+++ b/Engine/Plugins/VirtualProduction/DataLink/Source/DataLinkEditor/Private/Preview/DataLinkPreviewTool.cpp
@@ -110,8 +110,6 @@ void FDataLinkPreviewTool::RunPreview()
 #if WITH_DATALINK_CONTEXT
 		.SetContextName(TEXT("Graph Editor Preview"))
 #endif
-		// Blueprint nodes could use nodes like 'Delay' that require a world context object. Ensure preview can run these nodes by setting the context to GWorld
-		.SetContextObject(GWorld)
 		.SetSink(Sink)
 		.SetOnOutputData(FOnDataLinkOutputData::CreateSP(this, &FDataLinkPreviewTool::OnPreviewOutputData)));
 
