Fixes for CL 51336460 - Remove string duplication in FJsonObject to free memory

#rnx

#Synced-CL 51336291

diff --git a/Engine/Plugins/VirtualProduction/DataLink/Source/DataLinkJsonEditor/Private/DataLinkJsonEditorStructGenerator.cpp b/Engine/Plugins/VirtualProduction/DataLink/Source/DataLinkJsonEditor/Private/DataLinkJsonEditorStructGenerator.cpp
--- a/Engine/Plugins/VirtualProduction/DataLink/Source/DataLinkJsonEditor/Private/DataLinkJsonEditorStructGenerator.cpp
+++ b/Engine/Plugins/VirtualProduction/DataLink/Source/DataLinkJsonEditor/Private/DataLinkJsonEditorStructGenerator.cpp
@@ -200,9 +200,9 @@ UUserDefinedStruct* UE::DataLinkJsonEditor::FStructGenerator::GetOrCreateStruct(
 		return nullptr;
 	}
 
-	for (const TPair<FString, TSharedPtr<FJsonValue>>& JsonEntry : InJsonObject->Values)
+	for (const TPair<UE::FSharedString, TSharedPtr<FJsonValue>>& JsonEntry : InJsonObject->Values)
 	{
-		AddVariable(Struct, JsonEntry.Key, JsonEntry.Value);
+		AddVariable(Struct, FString(JsonEntry.Key), JsonEntry.Value);
 	}
 
 	FStructureEditorUtils::CompileStructure(Struct);
diff --git a/Engine/Source/Developer/IoStore/OnDemandUtilities/Private/Upload.cpp b/Engine/Source/Developer/IoStore/OnDemandUtilities/Private/Upload.cpp
--- a/Engine/Source/Developer/IoStore/OnDemandUtilities/Private/Upload.cpp
+++ b/Engine/Source/Developer/IoStore/OnDemandUtilities/Private/Upload.cpp
@@ -628,10 +628,10 @@ static TIoStatusOr<FUploadResult> UploadContainerFiles(
 
 
 
-		for (const TPair<FString, TSharedPtr<FJsonValue>>& Pair : PackageSetsObject->Values)
+		for (const TPair<UE::FSharedString, TSharedPtr<FJsonValue>>& Pair : PackageSetsObject->Values)
 		{
 			TArray<FString> Packages;
-			if (!PackageSetsObject->TryGetStringArrayField(Pair.Key, Packages))
+			if (!PackageSetsObject->TryGetStringArrayField(*Pair.Key, Packages))
 			{
 				UE_LOG(LogIas, Error, TEXT("Bad settings file '%s'"), *UploadParams.TagSetFile);
 				return FIoStatus(EIoErrorCode::ReadError, FString::Printf(TEXT("Bad settings file '%s'"), *UploadParams.TagSetFile));
